Fix for audio not coming back after an interruption by Siri, Alarm, or phone call
TTP329898
#ue4
#ios

diff --git a/Engine/Source/Runtime/IOS/IOSAudio/Public/IOSAudioDevice.h b/Engine/Source/Runtime/IOS/IOSAudio/Public/IOSAudioDevice.h
--- a/Engine/Source/Runtime/IOS/IOSAudio/Public/IOSAudioDevice.h
+++ b/Engine/Source/Runtime/IOS/IOSAudio/Public/IOSAudioDevice.h
@@ -189,8 +189,8 @@ public:
 	AudioUnit GetMixerUnit() const { return MixerUnit; }
 
 	/** Thread context management */
-	static void ResumeContext();
-	static void SuspendContext();
+	virtual void ResumeContext();
+	virtual void SuspendContext();
 	
 	virtual bool Exec( UWorld* InWorld, const TCHAR* Cmd, FOutputDevice& Ar = *GLog ) OVERRIDE;
 
