- MetalRHI, First pass, NOT COMPLETE/CLEANED UP
- Nothing compiles this yet, this is a staging checkin. No MetalRHI.Build.cs on purpose!

diff --git a/Engine/Source/Runtime/IOS/MetalRHI/Private/MetalCommands.cpp b/Engine/Source/Runtime/IOS/MetalRHI/Private/MetalCommands.cpp
new file mode 100644
--- /dev/null
+++ b/Engine/Source/Runtime/IOS/MetalRHI/Private/MetalCommands.cpp
@@ -0,0 +1,594 @@
+// Copyright 1998-2014 Epic Games, Inc. All Rights Reserved.
+
+/*=============================================================================
+	MetalCommands.cpp: Metal RHI commands implementation.
+=============================================================================*/
+
+#include "MetalRHIPrivate.h"
+
+#define NO_DRAW 0
+
+static MTLPrimitiveType TranslatePrimitiveType(uint32 PrimitiveType)
+{
+	switch (PrimitiveType)
+	{
+		case PT_TriangleList:	return MTLPrimitiveTypeTriangle;
+		case PT_TriangleStrip:	return MTLPrimitiveTypeTriangleStrip;
+		case PT_LineList:		return MTLPrimitiveTypeLine;
+		default:				UE_LOG(LogMetal, Fatal, TEXT("Unsupported primitive type %d"), (int32)PrimitiveType); return MTLPrimitiveTypeTriangle;
+	}
+}
+
+TRefCountPtr<FMetalVertexBuffer> HACK_ZeroStrideBuffers[16] = { 0 };
+
+void FMetalDynamicRHI::RHISetStreamSource(uint32 StreamIndex,FVertexBufferRHIParamRef VertexBufferRHI,uint32 Stride,uint32 Offset)
+{
+	DYNAMIC_CAST_METGALRESOURCE(VertexBuffer,VertexBuffer);
+
+//	FMetalManager::Get()->CheckStride(StreamIndex, Stride);
+
+	// queue this zero stride buffer to be grown later if needed
+	HACK_ZeroStrideBuffers[StreamIndex] = VertexBuffer->ZeroStrideElementSize ? VertexBuffer : NULL;
+
+	[FMetalManager::GetContext() setVertexBuffer:VertexBuffer->Buffer offset:VertexBuffer->Offset + Offset atIndex:UNREAL_TO_METAL_BUFFER_INDEX(StreamIndex)];
+}
+
+void FMetalDynamicRHI::RHISetStreamOutTargets(uint32 NumTargets, const FVertexBufferRHIParamRef* VertexBuffers, const uint32* Offsets)
+{
+	NOT_SUPPORTED("RHISetStreamOutTargets");
+
+}
+
+void FMetalDynamicRHI::RHISetRasterizerState(FRasterizerStateRHIParamRef NewStateRHI)
+{
+	DYNAMIC_CAST_METGALRESOURCE(RasterizerState,NewState);
+
+	NewState->Set();
+}
+
+void FMetalDynamicRHI::RHISetComputeShader(FComputeShaderRHIParamRef ComputeShaderRHI)
+{
+	NOT_SUPPORTED("RHISetComputeShader");
+}
+
+void FMetalDynamicRHI::RHIDispatchComputeShader(uint32 ThreadGroupCountX, uint32 ThreadGroupCountY, uint32 ThreadGroupCountZ) 
+{ 
+	NOT_SUPPORTED("RHIDispatchComputeShader");
+}
+
+void FMetalDynamicRHI::RHIDispatchIndirectComputeShader(FVertexBufferRHIParamRef ArgumentBufferRHI, uint32 ArgumentOffset) 
+{ 
+	NOT_SUPPORTED("RHIDispatchIndirectComputeShader");
+}
+
+void FMetalDynamicRHI::RHISetViewport(uint32 MinX,uint32 MinY,float MinZ,uint32 MaxX,uint32 MaxY,float MaxZ)
+{
+	MTLViewport Viewport;
+	Viewport.originX = MinX;
+	Viewport.originY = MinY;
+	Viewport.width = MaxX - MinX;
+	Viewport.height = MaxY - MinY;
+	Viewport.znear = MinZ;
+	Viewport.zfar = MaxZ;
+	
+	[FMetalManager::GetContext() setViewport:Viewport];
+}
+
+void FMetalDynamicRHI::RHISetMultipleViewports(uint32 Count, const FViewportBounds* Data) 
+{ 
+	NOT_SUPPORTED("RHISetMultipleViewports");
+}
+
+void FMetalDynamicRHI::RHISetScissorRect(bool bEnable,uint32 MinX,uint32 MinY,uint32 MaxX,uint32 MaxY)
+{
+	MTLScissorRect Scissor;
+	Scissor.x = MinX;
+	Scissor.y = MinY;
+	Scissor.width = MaxX - MinX;
+	Scissor.height = MaxY - MinY;
+
+	[FMetalManager::GetContext() setScissorRect:Scissor];
+}
+
+void FMetalDynamicRHI::RHISetBoundShaderState( FBoundShaderStateRHIParamRef BoundShaderStateRHI)
+{
+	DYNAMIC_CAST_METGALRESOURCE(BoundShaderState,BoundShaderState);
+
+	FMetalManager::Get()->SetBoundShaderState(BoundShaderState);
+	BoundShaderStateHistory.Add(BoundShaderState);
+}
+
+
+void FMetalDynamicRHI::RHISetUAVParameter(FComputeShaderRHIParamRef ComputeShaderRHI, uint32 UAVIndex, FUnorderedAccessViewRHIParamRef UAVRHI)
+{
+	NOT_SUPPORTED("RHISetUAVParameter");
+}
+
+void FMetalDynamicRHI::RHISetUAVParameter(FComputeShaderRHIParamRef ComputeShaderRHI,uint32 UAVIndex,FUnorderedAccessViewRHIParamRef UAVRHI, uint32 InitialCount)
+{
+	NOT_SUPPORTED("RHISetUAVParameter");
+}
+
+
+void FMetalDynamicRHI::RHISetShaderTexture(FVertexShaderRHIParamRef VertexShaderRHI, uint32 TextureIndex, FTextureRHIParamRef NewTextureRHI)
+{
+	FMetalSurface& Surface = GetMetalSurfaceFromRHITexture(NewTextureRHI);
+	[FMetalManager::GetContext() setVertexTexture:Surface.Texture atIndex:TextureIndex];
+}
+
+void FMetalDynamicRHI::RHISetShaderTexture(FHullShaderRHIParamRef HullShader, uint32 TextureIndex, FTextureRHIParamRef NewTextureRHI)
+{
+	NOT_SUPPORTED("RHISetShaderTexture-Hull");
+}
+
+void FMetalDynamicRHI::RHISetShaderTexture(FDomainShaderRHIParamRef DomainShader, uint32 TextureIndex, FTextureRHIParamRef NewTextureRHI)
+{
+	NOT_SUPPORTED("RHISetShaderTexture-Domain");
+}
+
+void FMetalDynamicRHI::RHISetShaderTexture(FGeometryShaderRHIParamRef GeometryShader, uint32 TextureIndex, FTextureRHIParamRef NewTextureRHI)
+{
+	NOT_SUPPORTED("RHISetShaderTexture-Geometry");
+
+}
+
+void FMetalDynamicRHI::RHISetShaderTexture(FPixelShaderRHIParamRef PixelShader, uint32 TextureIndex, FTextureRHIParamRef NewTextureRHI)
+{
+	FMetalSurface& Surface = GetMetalSurfaceFromRHITexture(NewTextureRHI);
+	[FMetalManager::GetContext() setFragmentTexture:Surface.Texture atIndex:TextureIndex];
+}
+
+void FMetalDynamicRHI::RHISetShaderTexture(FComputeShaderRHIParamRef ComputeShader, uint32 TextureIndex, FTextureRHIParamRef NewTextureRHI)
+{
+	NOT_SUPPORTED("RHISetShaderTexture-Compute");
+}
+
+
+void FMetalDynamicRHI::RHISetShaderResourceViewParameter(FVertexShaderRHIParamRef VertexShaderRHI, uint32 TextureIndex, FShaderResourceViewRHIParamRef SRVRHI)
+{
+	NOT_SUPPORTED("RHISetShaderResourceViewParameter");
+}
+
+void FMetalDynamicRHI::RHISetShaderResourceViewParameter(FHullShaderRHIParamRef HullShaderRHI,uint32 TextureIndex,FShaderResourceViewRHIParamRef SRVRHI)
+{
+	NOT_SUPPORTED("RHISetShaderResourceViewParameter");
+}
+
+void FMetalDynamicRHI::RHISetShaderResourceViewParameter(FDomainShaderRHIParamRef DomainShaderRHI,uint32 TextureIndex,FShaderResourceViewRHIParamRef SRVRHI)
+{
+	NOT_SUPPORTED("RHISetShaderResourceViewParameter");
+}
+
+void FMetalDynamicRHI::RHISetShaderResourceViewParameter(FGeometryShaderRHIParamRef GeometryShaderRHI,uint32 TextureIndex,FShaderResourceViewRHIParamRef SRVRHI)
+{
+	NOT_SUPPORTED("RHISetShaderResourceViewParameter");
+}
+
+void FMetalDynamicRHI::RHISetShaderResourceViewParameter(FPixelShaderRHIParamRef PixelShaderRHI,uint32 TextureIndex,FShaderResourceViewRHIParamRef SRVRHI)
+{
+	NOT_SUPPORTED("RHISetShaderResourceViewParameter");
+}
+
+void FMetalDynamicRHI::RHISetShaderResourceViewParameter(FComputeShaderRHIParamRef ComputeShaderRHI,uint32 TextureIndex,FShaderResourceViewRHIParamRef SRVRHI)
+{
+	NOT_SUPPORTED("RHISetShaderResourceViewParameter");
+}
+
+
+void FMetalDynamicRHI::RHISetShaderSampler(FVertexShaderRHIParamRef VertexShaderRHI, uint32 SamplerIndex, FSamplerStateRHIParamRef NewStateRHI)
+{
+	DYNAMIC_CAST_METGALRESOURCE(SamplerState,NewState);
+
+	[FMetalManager::GetContext() setVertexSamplerState:NewState->State atIndex:SamplerIndex];
+}
+
+void FMetalDynamicRHI::RHISetShaderSampler(FHullShaderRHIParamRef HullShader, uint32 SamplerIndex, FSamplerStateRHIParamRef NewStateRHI)
+{
+	DYNAMIC_CAST_METGALRESOURCE(SamplerState,NewState);
+
+	NOT_SUPPORTED("RHISetSamplerState-Hull");
+}
+
+void FMetalDynamicRHI::RHISetShaderSampler(FDomainShaderRHIParamRef DomainShader, uint32 SamplerIndex, FSamplerStateRHIParamRef NewStateRHI)
+{
+	DYNAMIC_CAST_METGALRESOURCE(SamplerState,NewState);
+
+	NOT_SUPPORTED("RHISetSamplerState-Domain");
+}
+
+void FMetalDynamicRHI::RHISetShaderSampler(FGeometryShaderRHIParamRef GeometryShader, uint32 SamplerIndex, FSamplerStateRHIParamRef NewStateRHI)
+{
+	DYNAMIC_CAST_METGALRESOURCE(SamplerState,NewState);
+
+	NOT_SUPPORTED("RHISetSamplerState-Geometry");
+}
+
+void FMetalDynamicRHI::RHISetShaderSampler(FPixelShaderRHIParamRef PixelShader, uint32 SamplerIndex, FSamplerStateRHIParamRef NewStateRHI)
+{
+	DYNAMIC_CAST_METGALRESOURCE(SamplerState,NewState);
+
+	[FMetalManager::GetContext() setFragmentSamplerState:NewState->State atIndex:SamplerIndex];
+}
+
+void FMetalDynamicRHI::RHISetShaderSampler(FComputeShaderRHIParamRef ComputeShader, uint32 SamplerIndex, FSamplerStateRHIParamRef NewStateRHI)
+{
+	DYNAMIC_CAST_METGALRESOURCE(SamplerState,NewState);
+
+	NOT_SUPPORTED("RHISetSamplerState-Compute");
+}
+
+
+void FMetalDynamicRHI::RHISetShaderParameter(FVertexShaderRHIParamRef VertexShaderRHI, uint32 BufferIndex, uint32 BaseIndex, uint32 NumBytes, const void* NewValue)
+{
+	FMetalManager::Get()->ShaderParameters[METAL_SHADER_STAGE_VERTEX].Set(BufferIndex, BaseIndex, NumBytes, NewValue);
+}
+
+void FMetalDynamicRHI::RHISetShaderParameter(FHullShaderRHIParamRef HullShaderRHI, uint32 BufferIndex, uint32 BaseIndex, uint32 NumBytes, const void* NewValue)
+{
+	NOT_SUPPORTED("RHISetShaderParameter-Hull");
+
+}
+
+void FMetalDynamicRHI::RHISetShaderParameter(FPixelShaderRHIParamRef PixelShaderRHI, uint32 BufferIndex, uint32 BaseIndex, uint32 NumBytes, const void* NewValue)
+{
+	FMetalManager::Get()->ShaderParameters[METAL_SHADER_STAGE_PIXEL].Set(BufferIndex, BaseIndex, NumBytes, NewValue);
+}
+
+void FMetalDynamicRHI::RHISetShaderParameter(FDomainShaderRHIParamRef DomainShaderRHI, uint32 BufferIndex, uint32 BaseIndex, uint32 NumBytes, const void* NewValue)
+{
+	NOT_SUPPORTED("RHISetShaderParameter-Domain");
+
+}
+
+void FMetalDynamicRHI::RHISetShaderParameter(FGeometryShaderRHIParamRef GeometryShaderRHI, uint32 BufferIndex, uint32 BaseIndex, uint32 NumBytes, const void* NewValue)
+{
+	NOT_SUPPORTED("RHISetShaderParameter-Geometry");
+
+}
+
+void FMetalDynamicRHI::RHISetShaderParameter(FComputeShaderRHIParamRef ComputeShaderRHI,uint32 BufferIndex, uint32 BaseIndex, uint32 NumBytes, const void* NewValue)
+{
+	NOT_SUPPORTED("RHISetShaderParameter-Compute");
+}
+
+void FMetalDynamicRHI::RHISetShaderUniformBuffer(FVertexShaderRHIParamRef VertexShaderRHI, uint32 BufferIndex, FUniformBufferRHIParamRef BufferRHI)
+{
+	DYNAMIC_CAST_METGALRESOURCE(VertexShader, VertexShader);
+	VertexShader->BoundUniformBuffers[BufferIndex] = BufferRHI;
+}
+
+void FMetalDynamicRHI::RHISetShaderUniformBuffer(FHullShaderRHIParamRef HullShader, uint32 BufferIndex, FUniformBufferRHIParamRef BufferRHI)
+{
+	NOT_SUPPORTED("RHISetShaderUniformBuffer-Hull");
+}
+
+void FMetalDynamicRHI::RHISetShaderUniformBuffer(FDomainShaderRHIParamRef DomainShader, uint32 BufferIndex, FUniformBufferRHIParamRef BufferRHI)
+{
+	NOT_SUPPORTED("RHISetShaderUniformBuffer-Domain");
+}
+
+void FMetalDynamicRHI::RHISetShaderUniformBuffer(FGeometryShaderRHIParamRef GeometryShader, uint32 BufferIndex, FUniformBufferRHIParamRef BufferRHI)
+{
+	NOT_SUPPORTED("RHISetShaderUniformBuffer-Geometry");
+}
+
+void FMetalDynamicRHI::RHISetShaderUniformBuffer(FPixelShaderRHIParamRef PixelShaderRHI, uint32 BufferIndex, FUniformBufferRHIParamRef BufferRHI)
+{
+	DYNAMIC_CAST_METGALRESOURCE(PixelShader, PixelShader);
+	PixelShader->BoundUniformBuffers[BufferIndex] = BufferRHI;
+}
+
+void FMetalDynamicRHI::RHISetShaderUniformBuffer(FComputeShaderRHIParamRef ComputeShader, uint32 BufferIndex, FUniformBufferRHIParamRef BufferRHI)
+{
+	NOT_SUPPORTED("RHISetShaderUniformBuffer-Compute");
+}
+
+
+void FMetalDynamicRHI::RHISetDepthStencilState(FDepthStencilStateRHIParamRef NewStateRHI, uint32 StencilRef)
+{
+	DYNAMIC_CAST_METGALRESOURCE(DepthStencilState,NewState);
+
+	NewState->Set();
+}
+
+void FMetalDynamicRHI::RHISetBlendState(FBlendStateRHIParamRef NewStateRHI, const FLinearColor& BlendFactor)
+{
+	DYNAMIC_CAST_METGALRESOURCE(BlendState,NewState);
+
+	NewState->Set();
+}
+
+
+void FMetalDynamicRHI::RHISetRenderTargets(uint32 NumSimultaneousRenderTargets, const FRHIRenderTargetView* NewRenderTargets,
+	FTextureRHIParamRef NewDepthStencilTargetRHI, uint32 NumUAVs, const FUnorderedAccessViewRHIParamRef* UAVs)
+{
+	FMetalManager* Manager = FMetalManager::Get();
+	
+	// only support 1 for now!!!
+	check(NumSimultaneousRenderTargets <= 1);
+	
+	// update the current RTs
+	if (NumSimultaneousRenderTargets && NewRenderTargets[0].Texture != NULL)
+	{
+		FMetalSurface& Surface = GetMetalSurfaceFromRHITexture(NewRenderTargets[0].Texture);
+		Manager->SetCurrentRenderTarget(&Surface);
+	}
+	else
+	{
+		Manager->SetCurrentRenderTarget(NULL);
+	}
+	
+	if (NewDepthStencilTargetRHI)
+	{
+		FMetalSurface& Surface = GetMetalSurfaceFromRHITexture(NewDepthStencilTargetRHI);
+		Manager->SetCurrentDepthStencilTarget(&Surface);
+	}
+	else
+	{
+		Manager->SetCurrentDepthStencilTarget(NULL);
+	}
+	
+	// now that we have a new render target, we need a new context to render to it!
+	Manager->UpdateContext();
+	
+	// Set the viewport to the full size of render target 0.
+	if (NewRenderTargets && NewRenderTargets[0].Texture)
+	{
+		FMetalSurface& RenderTarget = GetMetalSurfaceFromRHITexture(NewRenderTargets[0].Texture);
+		RHISetViewport(0, 0, 0.0f, RenderTarget.Texture.width, RenderTarget.Texture.height, 1.0f);
+	}
+}
+
+void FMetalDynamicRHI::RHIDiscardRenderTargets(bool Depth, bool Stencil, uint32 ColorBitMask)
+{
+}
+
+
+// Occlusion/Timer queries.
+void FMetalDynamicRHI::RHIBeginRenderQuery(FRenderQueryRHIParamRef QueryRHI)
+{
+	DYNAMIC_CAST_METGALRESOURCE(RenderQuery,Query);
+
+	Query->Begin();
+}
+
+void FMetalDynamicRHI::RHIEndRenderQuery(FRenderQueryRHIParamRef QueryRHI)
+{
+	DYNAMIC_CAST_METGALRESOURCE(RenderQuery,Query);
+
+	Query->End();
+}
+
+
+void FMetalDynamicRHI::RHIDrawPrimitive(uint32 PrimitiveType, uint32 BaseVertexIndex, uint32 NumPrimitives, uint32 NumInstances)
+{
+	SCOPE_CYCLE_COUNTER(STAT_MetalDrawCallTime);
+	checkf(NumInstances == 1, TEXT("Currently only 1 instance is supported"));
+	
+	if (BaseVertexIndex > 0)
+	{
+		UE_LOG(LogMetal, Log, TEXT("Attempted to render with BaseVertexIndex - skipping..."));
+		return;
+	}
+
+	// how many verts to render
+	uint32 NumVertices = RHIGetVertexCountForPrimitiveCount(NumPrimitives, PrimitiveType);
+
+	// finalize any pending state
+	FMetalManager::Get()->PrepareToDraw(NumVertices);
+
+	// draw!
+#if !NO_DRAW
+	[FMetalManager::GetContext() drawPrimitives:TranslatePrimitiveType(PrimitiveType)
+									vertexStart:BaseVertexIndex
+									vertexCount:NumVertices
+								  instanceCount:NumInstances];
+#endif
+}
+
+void FMetalDynamicRHI::RHIDrawPrimitiveIndirect(uint32 PrimitiveType, FVertexBufferRHIParamRef ArgumentBufferRHI, uint32 ArgumentOffset)
+{
+	NOT_SUPPORTED("RHIDrawPrimitiveIndirect");
+}
+
+
+void FMetalDynamicRHI::RHIDrawIndexedPrimitive(FIndexBufferRHIParamRef IndexBufferRHI, uint32 PrimitiveType, int32 BaseVertexIndex, uint32 MinIndex,
+	uint32 NumVertices, uint32 StartIndex, uint32 NumPrimitives, uint32 NumInstances)
+{
+	SCOPE_CYCLE_COUNTER(STAT_MetalDrawCallTime);
+	checkf(NumInstances == 1, TEXT("Currently only 1 instance is supported"));
+	if (BaseVertexIndex > 0)
+	{
+		UE_LOG(LogMetal, Log, TEXT("Attempted to render with BaseVertexIndex - skipping..."));
+		return;
+	}
+
+	DYNAMIC_CAST_METGALRESOURCE(IndexBuffer,IndexBuffer);
+
+	// finalize any pending state
+	FMetalManager::Get()->PrepareToDraw(NumVertices);
+	
+	uint32 NumIndices = RHIGetVertexCountForPrimitiveCount(NumPrimitives, PrimitiveType);
+#if !NO_DRAW
+	[FMetalManager::GetContext() drawIndexedPrimitives:TranslatePrimitiveType(PrimitiveType)
+											indexCount:NumIndices
+											 indexType:IndexBuffer->IndexType
+										   indexBuffer:IndexBuffer->Buffer
+									 indexBufferOffset:IndexBuffer->Offset + StartIndex * IndexBuffer->GetStride()
+										 instanceCount:NumInstances];
+#endif
+}
+
+void FMetalDynamicRHI::RHIDrawIndexedIndirect(FIndexBufferRHIParamRef IndexBufferRHI, uint32 PrimitiveType, FStructuredBufferRHIParamRef ArgumentsBufferRHI, int32 DrawArgumentsIndex, uint32 NumInstances)
+{
+	NOT_SUPPORTED("RHIDrawIndexedIndirect");
+}
+
+void FMetalDynamicRHI::RHIDrawIndexedPrimitiveIndirect(uint32 PrimitiveType,FIndexBufferRHIParamRef IndexBufferRHI,FVertexBufferRHIParamRef ArgumentBufferRHI,uint32 ArgumentOffset)
+{
+	NOT_SUPPORTED("RHIDrawIndexedPrimitiveIndirect");
+}
+
+
+/** Some locally global variables to track the pending primitive information uised in RHIEnd*UP functions */
+static uint32 GPendingVertexBufferOffset = 0xFFFFFFFF;
+static uint32 GPendingVertexDataStride;
+
+static uint32 GPendingIndexBufferOffset = 0xFFFFFFFF;
+static uint32 GPendingIndexDataStride;
+
+static uint32 GPendingPrimitiveType;
+static uint32 GPendingNumPrimitives;
+
+void FMetalDynamicRHI::RHIBeginDrawPrimitiveUP( uint32 PrimitiveType, uint32 NumPrimitives, uint32 NumVertices, uint32 VertexDataStride, void*& OutVertexData)
+{
+	checkSlow(GPendingVertexBufferOffset == 0xFFFFFFFF);
+
+	// allocate space
+	GPendingVertexBufferOffset = FMetalManager::Get()->AllocateFromRingBuffer(VertexDataStride * NumVertices);
+
+	// get the pointer to send back for writing
+	uint8* RingBufferBytes = (uint8*)[FMetalManager::Get()->GetRingBuffer() contents];
+	OutVertexData = RingBufferBytes + GPendingVertexBufferOffset;
+	
+	GPendingPrimitiveType = PrimitiveType;
+	GPendingNumPrimitives = NumPrimitives;
+	GPendingVertexDataStride = VertexDataStride;
+}
+
+
+void FMetalDynamicRHI::RHIEndDrawPrimitiveUP()
+{
+	SCOPE_CYCLE_COUNTER(STAT_MetalDrawCallTime);
+
+	// set the vertex buffer
+	[FMetalManager::GetContext() setVertexBuffer:FMetalManager::Get()->GetRingBuffer() offset:GPendingVertexBufferOffset atIndex:UNREAL_TO_METAL_BUFFER_INDEX(0)];
+	
+	// how many to draw
+	uint32 NumVertices = RHIGetVertexCountForPrimitiveCount(GPendingNumPrimitives, GPendingPrimitiveType);
+
+	// last minute draw setup
+	HACK_ZeroStrideBuffers[0] = NULL;
+	FMetalManager::Get()->PrepareToDraw(0);
+	FMetalManager::Get()->CheckStride(0, GPendingVertexDataStride);
+	
+#if !NO_DRAW
+	[FMetalManager::GetContext() drawPrimitives:TranslatePrimitiveType(GPendingPrimitiveType)
+									vertexStart:0
+									vertexCount:NumVertices
+								  instanceCount:1];
+#endif
+	
+	// mark temp memory as usable
+	GPendingVertexBufferOffset = 0xFFFFFFFF;
+}
+
+void FMetalDynamicRHI::RHIBeginDrawIndexedPrimitiveUP( uint32 PrimitiveType, uint32 NumPrimitives, uint32 NumVertices, uint32 VertexDataStride, void*& OutVertexData, uint32 MinVertexIndex, uint32 NumIndices, uint32 IndexDataStride, void*& OutIndexData)
+{
+	checkSlow(GPendingVertexBufferOffset == 0xFFFFFFFF);
+	checkSlow(GPendingIndexBufferOffset == 0xFFFFFFFF);
+
+	// allocate space
+	GPendingVertexBufferOffset = FMetalManager::Get()->AllocateFromRingBuffer(VertexDataStride * NumVertices);
+	GPendingIndexBufferOffset = FMetalManager::Get()->AllocateFromRingBuffer(IndexDataStride * NumIndices);
+	
+	// get the pointers to send back for writing
+	uint8* RingBufferBytes = (uint8*)[FMetalManager::Get()->GetRingBuffer() contents];
+	OutVertexData = RingBufferBytes + GPendingVertexBufferOffset;
+	OutIndexData = RingBufferBytes + GPendingIndexBufferOffset;
+	
+	GPendingPrimitiveType = PrimitiveType;
+	GPendingNumPrimitives = NumPrimitives;
+	GPendingIndexDataStride = IndexDataStride;
+
+	GPendingVertexDataStride = VertexDataStride;
+}
+
+void FMetalDynamicRHI::RHIEndDrawIndexedPrimitiveUP()
+{
+	SCOPE_CYCLE_COUNTER(STAT_MetalDrawCallTime);
+
+	// set the vertex buffer
+	[FMetalManager::GetContext() setVertexBuffer:FMetalManager::Get()->GetRingBuffer() offset:GPendingVertexBufferOffset atIndex:UNREAL_TO_METAL_BUFFER_INDEX(0)];
+
+	// how many to draw
+	uint32 NumIndices = RHIGetVertexCountForPrimitiveCount(GPendingNumPrimitives, GPendingPrimitiveType);
+
+	// last minute draw setup
+	HACK_ZeroStrideBuffers[0] = NULL;
+	FMetalManager::Get()->PrepareToDraw(0);
+	FMetalManager::Get()->CheckStride(0, GPendingVertexDataStride);
+	
+#if !NO_DRAW
+	[FMetalManager::GetContext() drawIndexedPrimitives:TranslatePrimitiveType(GPendingPrimitiveType)
+											indexCount:NumIndices
+											 indexType:(GPendingIndexDataStride == 2) ? MTLIndexTypeUInt16 : MTLIndexTypeUInt32
+										   indexBuffer:FMetalManager::Get()->GetRingBuffer()
+									 indexBufferOffset:GPendingIndexBufferOffset
+										 instanceCount:1];
+#endif
+	
+	// mark temp memory as usable
+	GPendingVertexBufferOffset = 0xFFFFFFFF;
+	GPendingIndexBufferOffset = 0xFFFFFFFF;
+}
+
+
+void FMetalDynamicRHI::RHIClear(bool bClearColor,const FLinearColor& Color, bool bClearDepth,float Depth, bool bClearStencil,uint32 Stencil, FIntRect ExcludeRect)
+{
+	FMetalDynamicRHI::RHIClearMRT(bClearColor, 1, &Color, bClearDepth, Depth, bClearStencil, Stencil, ExcludeRect);
+}
+
+void FMetalDynamicRHI::RHIClearMRT(bool bClearColor,int32 NumClearColors,const FLinearColor* ClearColorArray,bool bClearDepth,float Depth,bool bClearStencil,uint32 Stencil, FIntRect ExcludeRect)
+{
+	checkf(NumClearColors == 1, TEXT("Only one MRT supported in RHIClearMRT"));
+}
+
+
+
+void FMetalDynamicRHI::RHISuspendRendering()
+{
+	// Not supported
+}
+
+void FMetalDynamicRHI::RHIResumeRendering()
+{
+	// Not supported
+}
+
+bool FMetalDynamicRHI::RHIIsRenderingSuspended()
+{
+	// Not supported
+	return false;
+}
+
+void FMetalDynamicRHI::RHIBlockUntilGPUIdle()
+{
+	NOT_SUPPORTED("RHIBlockUntilGPUIdle");
+}
+
+uint32 FMetalDynamicRHI::RHIGetGPUFrameCycles()
+{
+	return GGPUFrameTime;
+}
+
+void FMetalDynamicRHI::RHIAutomaticCacheFlushAfterComputeShader(bool bEnable) 
+{
+	NOT_SUPPORTED("RHIAutomaticCacheFlushAfterComputeShader");
+}
+
+void FMetalDynamicRHI::RHIFlushComputeShaderCache()
+{
+	NOT_SUPPORTED("RHIFlushComputeShaderCache");
+}
+
+void FMetalDynamicRHI::RHIExecuteCommandList(FRHICommandList* RHICmdList)
+{
+	NOT_SUPPORTED("RHIExecuteCommandList");
+}
+
+void FMetalDynamicRHI::RHIEnableDepthBoundsTest(bool bEnable, float MinDepth, float MaxDepth)
+{
+	// not supported
+}
+
diff --git a/Engine/Source/Runtime/IOS/MetalRHI/Private/MetalGlobalUniformBuffer.h b/Engine/Source/Runtime/IOS/MetalRHI/Private/MetalGlobalUniformBuffer.h
new file mode 100644
--- /dev/null
+++ b/Engine/Source/Runtime/IOS/MetalRHI/Private/MetalGlobalUniformBuffer.h
@@ -0,0 +1,48 @@
+// Copyright 1998-2014 Epic Games, Inc. All Rights Reserved.
+
+/*=============================================================================
+	MetalGlobalUniformBuffer.h: Metal Global uniform definitions.
+=============================================================================*/
+
+#pragma once 
+
+/** Size of the default constant buffer (copied from D3D11) */
+#define MAX_GLOBAL_CONSTANT_BUFFER_SIZE		4096
+
+/**
+ * An Metal uniform buffer that has backing memory to store global uniforms
+ */
+class FMetalGlobalUniformBuffer : public FRenderResource, public FRefCountedObject
+{
+public:
+
+	FMetalGlobalUniformBuffer(uint32 InSize = 0);
+	~FMetalGlobalUniformBuffer();
+
+	// FRenderResource interface.
+	virtual void	InitDynamicRHI();
+	virtual void	ReleaseDynamicRHI();
+
+	
+	void UpdateConstant(const void* Data, uint16 Offset, uint16 Size);
+
+	void CommitToVertexShader(bool bDiscardSharedConstants);
+	void CommitToFragmentShader(bool bDiscardSharedConstants);
+
+protected:
+	
+	uint32 FinalizeBuffer(bool bDiscardSharedConstants);
+	
+	uint32	MaxSize;
+	bool	bIsDirty;
+	uint8*	ShadowData;
+
+	/** Size of all constants that has been updated since the last call to Commit. */
+	uint32	CurrentUpdateSize;
+
+	/**
+	 * Size of all constants that has been updated since the last Discard.
+	 * Includes "shared" constants that don't necessarily gets updated between every Commit.
+	 */
+	uint32	TotalUpdateSize;
+};
diff --git a/Engine/Source/Runtime/IOS/MetalRHI/Private/MetalIndexBuffer.cpp b/Engine/Source/Runtime/IOS/MetalRHI/Private/MetalIndexBuffer.cpp
new file mode 100644
--- /dev/null
+++ b/Engine/Source/Runtime/IOS/MetalRHI/Private/MetalIndexBuffer.cpp
@@ -0,0 +1,85 @@
+// Copyright 1998-2014 Epic Games, Inc. All Rights Reserved.
+
+/*=============================================================================
+	MetalIndexBuffer.cpp: Metal Index buffer RHI implementation.
+=============================================================================*/
+
+#include "MetalRHIPrivate.h"
+
+
+/** Constructor */
+FMetalIndexBuffer::FMetalIndexBuffer(uint32 InStride, uint32 InSize, uint32 InUsage)
+	: FRHIIndexBuffer(InStride, InSize, InUsage)
+	, IndexType((InStride == 2) ? MTLIndexTypeUInt16 : MTLIndexTypeUInt32)
+{
+	// volatile buffers can come right from the ring buffer
+	if (InUsage & BUF_Volatile)
+	{
+		Buffer = FMetalManager::Get()->GetRingBuffer();
+		// retaining to ease the destructor logic
+		[Buffer retain];
+	}
+	else
+	{
+		Buffer = [FMetalManager::GetDevice() newBufferWithLength:InSize options:BUFFER_CACHE_MODE];
+		Offset = 0;
+	}
+	TRACK_OBJECT(Buffer);
+}
+
+FMetalIndexBuffer::~FMetalIndexBuffer()
+{
+	FMetalManager::ReleaseObject(Buffer);
+}
+
+void* FMetalIndexBuffer::Lock(EResourceLockMode LockMode, uint32 Size)
+{
+	if (GetUsage() & BUF_Volatile)
+	{
+		Offset = FMetalManager::Get()->AllocateFromRingBuffer(Size);
+	}
+	return ((uint8*)[Buffer contents]) + Offset;
+}
+
+void FMetalIndexBuffer::Unlock()
+{
+
+}
+
+FIndexBufferRHIRef FMetalDynamicRHI::RHICreateIndexBuffer(uint32 Stride,uint32 Size, uint32 InUsage, FRHIResourceCreateInfo& CreateInfo)
+{
+	// make the RHI object, which will allocate memory
+	FMetalIndexBuffer* IndexBuffer = new FMetalIndexBuffer(Stride, Size, InUsage);
+	
+	if (CreateInfo.ResourceArray)
+	{
+		check(Size == CreateInfo.ResourceArray->GetResourceDataSize());
+
+		// make a buffer usable by CPU
+		void* Buffer = RHILockIndexBuffer(IndexBuffer, 0, Size, RLM_WriteOnly);
+
+		// copy the contents of the given data into the buffer
+		FMemory::Memcpy(Buffer, CreateInfo.ResourceArray->GetResourceData(), Size);
+
+		RHIUnlockIndexBuffer(IndexBuffer);
+
+		// Discard the resource array's contents.
+		CreateInfo.ResourceArray->Discard();
+	}
+
+	return IndexBuffer;
+}
+
+void* FMetalDynamicRHI::RHILockIndexBuffer(FIndexBufferRHIParamRef IndexBufferRHI, uint32 Offset, uint32 Size, EResourceLockMode LockMode)
+{
+	DYNAMIC_CAST_METGALRESOURCE(IndexBuffer,IndexBuffer);
+
+	return (uint8*)IndexBuffer->Lock(LockMode, Size) + Offset;
+}
+
+void FMetalDynamicRHI::RHIUnlockIndexBuffer(FIndexBufferRHIParamRef IndexBufferRHI)
+{
+	DYNAMIC_CAST_METGALRESOURCE(IndexBuffer,IndexBuffer);
+
+	IndexBuffer->Unlock();
+}
diff --git a/Engine/Source/Runtime/IOS/MetalRHI/Private/MetalManager.cpp b/Engine/Source/Runtime/IOS/MetalRHI/Private/MetalManager.cpp
new file mode 100644
--- /dev/null
+++ b/Engine/Source/Runtime/IOS/MetalRHI/Private/MetalManager.cpp
@@ -0,0 +1,845 @@
+// Copyright 1998-2014 Epic Games, Inc. All Rights Reserved.
+
+#include "MetalManager.h"
+#include "IOSAppDelegate.h"
+
+const uint32 RingBufferSize = 8 * 1024 * 1024;
+
+#if SHOULD_TRACK_OBJECTS
+TMap<id, int32> ClassCounts;
+#endif
+
+#define NUMBITS_DEPTH_WRITE_ENABLED				1
+#define NUMBITS_STENCIL_WRITE_ENABLED			1
+#define NUMBITS_SOURCE_RGB_BLEND_FACTOR			4
+#define NUMBITS_DEST_RGB_BLEND_FACTOR			4
+#define NUMBITS_RGB_BLEND_OPERATION				3
+#define NUMBITS_SOURCE_A_BLEND_FACTOR			4
+#define NUMBITS_DEST_A_BLEND_FACTOR				4
+#define NUMBITS_A_BLEND_OPERATION				3
+#define NUMBITS_WRITE_MASK						4
+#define NUMBITS_RENDER_TARGET_FORMAT			9
+#define NUMBITS_DEPTH_TARGET_FORMAT				9
+#define NUMBITS_SAMPLE_COUNT					3
+
+
+#define OFFSET_DEPTH_WRITE_ENABLED				(0)
+#define OFFSET_STENCIL_WRITE_ENABLED			(OFFSET_DEPTH_WRITE_ENABLED		+ NUMBITS_DEPTH_WRITE_ENABLED)
+#define OFFSET_SOURCE_RGB_BLEND_FACTOR			(OFFSET_STENCIL_WRITE_ENABLED	+ NUMBITS_STENCIL_WRITE_ENABLED)
+#define OFFSET_DEST_RGB_BLEND_FACTOR			(OFFSET_SOURCE_RGB_BLEND_FACTOR	+ NUMBITS_SOURCE_RGB_BLEND_FACTOR)
+#define OFFSET_RGB_BLEND_OPERATION				(OFFSET_DEST_RGB_BLEND_FACTOR	+ NUMBITS_DEST_RGB_BLEND_FACTOR)
+#define OFFSET_SOURCE_A_BLEND_FACTOR			(OFFSET_RGB_BLEND_OPERATION		+ NUMBITS_RGB_BLEND_OPERATION)
+#define OFFSET_DEST_A_BLEND_FACTOR				(OFFSET_SOURCE_A_BLEND_FACTOR	+ NUMBITS_SOURCE_A_BLEND_FACTOR)
+#define OFFSET_A_BLEND_OPERATION				(OFFSET_DEST_A_BLEND_FACTOR		+ NUMBITS_DEST_A_BLEND_FACTOR)
+#define OFFSET_WRITE_MASK						(OFFSET_A_BLEND_OPERATION		+ NUMBITS_A_BLEND_OPERATION)
+#define OFFSET_RENDER_TARGET_FORMAT				(OFFSET_WRITE_MASK				+ NUMBITS_WRITE_MASK)
+#define OFFSET_DEPTH_TARGET_FORMAT				(OFFSET_RENDER_TARGET_FORMAT	+ NUMBITS_RENDER_TARGET_FORMAT)
+#define OFFSET_SAMPLE_COUNT						(OFFSET_DEPTH_TARGET_FORMAT		+ NUMBITS_DEPTH_TARGET_FORMAT)
+
+#define SET_HASH(Offset, NumBits, Value) \
+	{ \
+		uint64 BitMask = ((1ULL << NumBits) - 1) << Offset; \
+		Pipeline.Hash = (Pipeline.Hash & ~BitMask) | (((uint64)Value << Offset) & BitMask); \
+	}
+#define GET_HASH(Offset, NumBits) ((Hash >> Offset) & ((1ULL << NumBits) - 1))
+
+extern FString GMaterialNameHack;
+DEFINE_STAT(STAT_MetalMakeDrawableTime);
+DEFINE_STAT(STAT_MetalDrawCallTime);
+DEFINE_STAT(STAT_MetalPrepareDrawTime);
+DEFINE_STAT(STAT_MetalUniformBufferCleanupTime);
+DEFINE_STAT(STAT_MetalFreeUniformBufferMemory);
+DEFINE_STAT(STAT_MetalNumFreeUniformBuffers);
+DEFINE_STAT(STAT_MetalPipelineStateTime);
+DEFINE_STAT(STAT_MetalBoundShaderStateTime);
+DEFINE_STAT(STAT_MetalVertexDeclarationTime);
+
+	
+void FPipelineShadow::SetHash(uint64 InHash)
+{
+	Hash = InHash;
+
+	bIsDepthWriteEnabled = GET_HASH(OFFSET_DEPTH_WRITE_ENABLED, NUMBITS_DEPTH_WRITE_ENABLED);
+	bIsStencilWriteEnabled = GET_HASH(OFFSET_STENCIL_WRITE_ENABLED, NUMBITS_STENCIL_WRITE_ENABLED);
+	BlendState[0] = [[MTLBlendDescriptor alloc] init];
+	FMetalManager::Get()->ReleaseObject(BlendState[0]);
+	BlendState[0].sourceRGBBlendFactor = (MTLBlendFactor)GET_HASH(OFFSET_SOURCE_RGB_BLEND_FACTOR, NUMBITS_SOURCE_RGB_BLEND_FACTOR);
+	BlendState[0].destinationRGBBlendFactor = (MTLBlendFactor)GET_HASH(OFFSET_DEST_RGB_BLEND_FACTOR, NUMBITS_DEST_RGB_BLEND_FACTOR);
+	BlendState[0].rgbBlendOperation = (MTLBlendOperation)GET_HASH(OFFSET_RGB_BLEND_OPERATION, NUMBITS_RGB_BLEND_OPERATION);
+	BlendState[0].sourceAlphaBlendFactor = (MTLBlendFactor)GET_HASH(OFFSET_SOURCE_A_BLEND_FACTOR, NUMBITS_SOURCE_A_BLEND_FACTOR);
+	BlendState[0].destinationAlphaBlendFactor = (MTLBlendFactor)GET_HASH(OFFSET_DEST_A_BLEND_FACTOR, NUMBITS_DEST_A_BLEND_FACTOR);
+	BlendState[0].alphaBlendOperation = (MTLBlendOperation)GET_HASH(OFFSET_A_BLEND_OPERATION, NUMBITS_A_BLEND_OPERATION);
+	BlendState[0].writeMask = GET_HASH(OFFSET_WRITE_MASK, NUMBITS_WRITE_MASK);
+	BlendState[0].blendingEnabled =
+		BlendState[0].rgbBlendOperation != MTLBlendOperationAdd || BlendState[0].destinationRGBBlendFactor != MTLBlendFactorZero || BlendState[0].sourceRGBBlendFactor != MTLBlendFactorOne ||
+		BlendState[0].alphaBlendOperation != MTLBlendOperationAdd || BlendState[0].destinationAlphaBlendFactor != MTLBlendFactorZero || BlendState[0].sourceAlphaBlendFactor != MTLBlendFactorOne;
+
+	RenderTargetFormat[0] = (MTLPixelFormat)GET_HASH(OFFSET_RENDER_TARGET_FORMAT, NUMBITS_RENDER_TARGET_FORMAT);
+	DepthTargetFormat = (MTLPixelFormat)GET_HASH(OFFSET_DEPTH_TARGET_FORMAT, NUMBITS_DEPTH_TARGET_FORMAT);
+	SampleCount = GET_HASH(OFFSET_SAMPLE_COUNT, NUMBITS_SAMPLE_COUNT);
+}
+
+
+
+static MTLTriangleFillMode TranslateFillMode(ERasterizerFillMode FillMode)
+{
+	switch (FillMode)
+	{
+		case FM_Wireframe:	return MTLTriangleFillModeLines;
+		case FM_Point:		return MTLTriangleFillModeFill;
+		default:			return MTLTriangleFillModeFill;
+	};
+}
+
+static MTLCullMode TranslateCullMode(ERasterizerCullMode CullMode)
+{
+	switch (CullMode)
+	{
+		case CM_CCW:	return MTLCullModeFront;
+		case CM_CW:		return MTLCullModeBack;
+		default:		return MTLCullModeNone;
+	}
+}
+
+
+
+
+id<MTLRenderPipelineState> FPipelineShadow::CreatePipelineStateForBoundShaderState(FMetalBoundShaderState* BSS) const
+{
+	SCOPE_CYCLE_COUNTER(STAT_MetalPipelineStateTime);
+	
+	MTLRenderPipelineDescriptor* Desc = [[MTLRenderPipelineDescriptor alloc] init];
+
+	// some basic settings
+	Desc.depthWriteEnabled = bIsDepthWriteEnabled;
+	Desc.stencilWriteEnabled = bIsStencilWriteEnabled;
+	if (FParse::Param(FCommandLine::Get(), TEXT("labels")))
+	{
+		Desc.label = GMaterialNameHack.GetNSString();
+	}
+
+	// set per-MRT settings
+	for (uint32 RenderTargetIndex = 0; RenderTargetIndex < MaxMetalRenderTargets; ++RenderTargetIndex)
+	{
+		[Desc setBlendDescriptor:BlendState[RenderTargetIndex] atIndex:(MTLFramebufferAttachmentIndex)RenderTargetIndex];
+		if (RenderTargetFormat[RenderTargetIndex] != MTLPixelFormatInvalid)
+		{
+			[Desc setPixelFormat:RenderTargetFormat[RenderTargetIndex] atIndex:(MTLFramebufferAttachmentIndex)RenderTargetIndex];
+		}
+	}
+
+	// depth setting if it's actually used
+	if (DepthTargetFormat != MTLPixelFormatInvalid)
+	{
+		[Desc setPixelFormat:DepthTargetFormat atIndex:MTLFramebufferAttachmentIndexDepth];
+	}
+
+	// set the bound shader state settings
+	Desc.vertexDescriptor = BSS->VertexDeclaration->Layout;
+	Desc.vertexFunction = BSS->VertexShader->Function;
+	Desc.fragmentFunction = BSS->PixelShader ? BSS->PixelShader->Function : nil;
+	Desc.sampleCount = SampleCount;
+
+	check(SampleCount > 0);
+
+	NSError* Error = nil;
+	id<MTLRenderPipelineState> PipelineState = [FMetalManager::GetDevice() newRenderPipelineStateWithDescriptor:Desc error : &Error];
+	TRACK_OBJECT(PipelineState);
+
+	[Desc release];
+
+	if (PipelineState == nil)
+	{
+		NSLog(@"Failed to generate a pipeline state object: %@", Error);
+		return nil;
+	}
+
+	return PipelineState;
+}
+
+FMetalManager* FMetalManager::Get()
+{
+	static FMetalManager Singleton;
+	return &Singleton;
+}
+
+id<MTLDevice> FMetalManager::GetDevice()
+{
+	return FMetalManager::Get()->Device;
+}
+
+id<MTLRenderCommandEncoder> FMetalManager::GetContext()
+{
+	return FMetalManager::Get()->CurrentContext;
+}
+
+void FMetalManager::ReleaseObject(id Object)
+{
+	FMetalManager::Get()->DelayedFreeLists[FMetalManager::Get()->WhichFreeList].Add(Object);
+}
+
+
+@interface MetalFramePacer : NSObject
+{
+@public
+	FEvent *FramePacerEvent;
+}
+
+-(void)run:(id)param;
+-(void)signal:(id)param;
+
+@end
+
+@implementation MetalFramePacer
+
+-(void)run:(id)param
+{
+	NSRunLoop *runloop = [NSRunLoop currentRunLoop];
+	CADisplayLink *displayLink = [CADisplayLink displayLinkWithTarget:self selector:@selector(signal:)];
+	displayLink.frameInterval = 2;
+	
+	[NSThread currentThread].name = @"FramePacerThread";
+	FramePacerEvent	= new FLocalEvent();
+	FramePacerEvent->Create(false);
+		
+	[displayLink addToRunLoop:runloop forMode:NSDefaultRunLoopMode];
+	[runloop run];
+}
+
+-(void)signal:(id)param
+{
+	FramePacerEvent->Trigger();
+}
+
+@end
+
+MetalFramePacer *FramePacer;
+
+FMetalManager::FMetalManager()
+	: Device([IOSAppDelegate GetDelegate].MetalView->Device)
+	, CurrentCommandBuffer(nil)
+	, CurrentDrawable(nil)
+	, CurrentFramebuffer(nil)
+	, CurrentContext(nil)
+	, CurrentColorRenderTexture(nil)
+	, PreviousColorRenderTexture(nil)
+	, CurrentDepthRenderTexture(nil)
+	, PreviousDepthRenderTexture(nil)
+	, RingBuffer(Device, RingBufferSize, BufferOffsetAlignment)
+	, QueryBuffer(Device, 64 * 1024, 8)
+	, WhichFreeList(0)
+	, CommandBufferIndex(0)
+	, CompletedCommandBufferIndex(0)
+
+{
+	CommandQueue = [Device newCommandQueue];
+
+	// get the size of the window
+	CGRect ViewFrame = [[IOSAppDelegate GetDelegate].MetalView frame];
+	FRHIResourceCreateInfo CreateInfo;
+	BackBuffer = (FMetalTexture2D*)(FTexture2DRHIParamRef)RHICreateTexture2D(ViewFrame.size.width, ViewFrame.size.height, PF_B8G8R8A8, 1, 1, TexCreate_RenderTargetable | TexCreate_Presentable, CreateInfo);
+
+//@todo-rco: What Size???
+	// make a buffer for each shader type
+	ShaderParameters = new FMetalShaderParameterCache[METAL_NUM_SHADER_STAGES];
+	ShaderParameters[METAL_SHADER_STAGE_VERTEX].InitializeResources(1024*1024);
+	ShaderParameters[METAL_SHADER_STAGE_PIXEL].InitializeResources(1024*1024);
+
+	// create a semaphore for multi-buffering the command buffer
+	CommandBufferSemaphore = dispatch_semaphore_create(FParse::Param(FCommandLine::Get(),TEXT("gpulockstep")) ? 1 : 3);
+
+	AutoReleasePoolTLSSlot = FPlatformTLS::AllocTlsSlot();
+
+#if GPU_TIMINGS
+	// start at start of GPUTimers
+	TimerIndex = 0;
+	FirstIndexToDump = 0;
+	LastIndexToDump = 0;
+#endif
+	
+	// Create display link thread
+	FramePacer = [[MetalFramePacer alloc] init];
+	[NSThread detachNewThreadSelector:@selector(run:) toTarget:FramePacer withObject:nil];
+	
+	InitFrame();
+}
+
+void FMetalManager::BeginFrame()
+{
+}
+
+
+static TConsoleVariableData<int32>* GDrawableMode = IConsoleManager::Get().RegisterConsoleVariable(TEXT("u.Drawable"), 1, TEXT("Whether to use new drawable mode or not"), ECVF_Default)->AsVariableInt();
+static bool GUseNewDrawable = true;
+
+void FMetalManager::InitFrame()
+{
+	GUseNewDrawable = GDrawableMode->GetValueOnRenderThread() > 0 ? true : false;
+
+	// start an auto release pool (EndFrame will drain and remake)
+	FPlatformTLS::SetTlsValue(AutoReleasePoolTLSSlot, FPlatformMisc::CreateAutoreleasePool());
+
+	// create the command buffer for this frame
+	CurrentCommandBuffer = [CommandQueue commandBufferWithUnretainedReferences];
+	[CurrentCommandBuffer retain];
+	TRACK_OBJECT(CurrentCommandBuffer);
+	
+#if GPU_TIMINGS
+	uint8 LocalTimerIndex = TimerIndex++;
+	FirstIndexThisFrame = LocalTimerIndex;
+#endif
+
+	uint64 LocalCommandBufferIndex = CommandBufferIndex++;
+	[CurrentCommandBuffer addScheduledHandler : ^ (id <MTLCommandBuffer> Buffer)
+	{
+#if GPU_TIMINGS
+		// timer stuff
+		GPUTimers[LocalTimerIndex] = FPlatformTime::Seconds();
+#endif
+
+		FMetalManager::Get()->SetCompletedCommandBufferIndex(LocalCommandBufferIndex);
+	}];
+
+//	double Start = FPlatformTime::Seconds();
+	// block on the semaphore
+	dispatch_semaphore_wait(CommandBufferSemaphore, DISPATCH_TIME_FOREVER);
+
+//	double Mid = FPlatformTime::Seconds();
+
+	if (GUseNewDrawable)
+	{
+		// mark us to get later
+		BackBuffer->Surface.Texture = nil;
+	}
+	else
+	{
+		SCOPE_CYCLE_COUNTER(STAT_MetalMakeDrawableTime);
+
+		// make a drawable object for this frame
+		CurrentDrawable = [[IOSAppDelegate GetDelegate].MetalView MakeDrawable];
+		TRACK_OBJECT(CurrentDrawable);
+		BackBuffer->Surface.Texture = CurrentDrawable.texture;
+	}
+
+//	double End = FPlatformTime::Seconds();
+//	NSLog(@"Semaphore Block Time: %.2f   -- MakeDrawable Time: %.2f", 1000.0f*(Mid - Start), 1000.0f*(End - Mid));
+	
+
+	extern void InitFrame_UniformBufferPoolCleanup();
+	InitFrame_UniformBufferPoolCleanup();
+
+	NumDrawCalls = 0;
+}
+
+void FMetalManager::EndFrame(bool bPresent)
+{
+//	NSLog(@"There were %d draw calls for final RT in frame %lld", NumDrawCalls, GFrameCounter);
+	// commit the render context to the commandBuffer
+	[CurrentContext endEncoding];
+	ReleaseObject(CurrentContext);
+	CurrentContext = nil;
+
+#if GPU_TIMINGS
+	// set up timer
+	uint8 LocalTimerIndex = TimerIndex++;
+	uint8 LocalFirstIndex = FirstIndexThisFrame;
+#endif
+
+	// kick the whole buffer
+	[CurrentCommandBuffer addCompletedHandler:^(id <MTLCommandBuffer> Buffer)
+	 {
+#if GPU_TIMINGS
+	 GPUTimers[LocalTimerIndex] = FPlatformTime::Seconds();
+	 FirstIndexToDump = LocalFirstIndex;
+	 LastIndexToDump = LocalTimerIndex;
+#endif
+	 dispatch_semaphore_signal(CommandBufferSemaphore);
+	 }];
+
+	// Wait until at least 2 VBlanks has passed since last time
+	FramePacer->FramePacerEvent->Wait();
+
+	// Commit before waiting to avoid leaving the gpu idle
+	[CurrentCommandBuffer commit];
+	
+	// enqueue a present if desired
+	if (CurrentDrawable)
+	{
+		if (bPresent && GFrameCounter > 3)
+		{
+			[CurrentCommandBuffer waitUntilScheduled];
+			
+			[CurrentDrawable present];
+		}
+
+		UNTRACK_OBJECT(CurrentDrawable);
+		[CurrentDrawable release];
+		CurrentDrawable = nil;
+	}
+
+
+	ReleaseObject(CurrentFramebuffer);
+	CurrentFramebuffer = nil;
+
+
+	UNTRACK_OBJECT(CurrentCommandBuffer);
+	[CurrentCommandBuffer release];
+
+#if GPU_TIMINGS
+	// print out timings
+	bool bDone = false;
+	uint8 Index = FirstIndexToDump;
+	uint8 RTIndex = 0;
+	while (Index != LastIndexToDump)
+	{
+		// this can wrap around at 256
+		uint8 NextIndex = Index + 1;
+		double Time = (GPUTimers[NextIndex] - GPUTimers[Index]) * 1000.0;
+		NSLog(@"Rende40rTarget %d took %.3fms", RTIndex, (float)Time);
+
+		Index = NextIndex;
+		RTIndex++;
+	}
+#endif
+
+	// xcode helper function
+	[CommandQueue insertDebugCaptureBoundary];
+
+	// drain the pool
+	FPlatformMisc::ReleaseAutoreleasePool(FPlatformTLS::GetTlsValue(AutoReleasePoolTLSSlot));
+
+	// drain the oldest pool
+	uint32 PrevFreeList = (WhichFreeList + 1) % ARRAY_COUNT(DelayedFreeLists);
+	for (int32 Index = 0; Index < DelayedFreeLists[PrevFreeList].Num(); Index++)
+	{
+		id Object = DelayedFreeLists[PrevFreeList][Index];
+		UNTRACK_OBJECT(Object);
+		[Object release];
+	}
+	DelayedFreeLists[PrevFreeList].Reset();
+
+	// and switch to it
+	WhichFreeList = PrevFreeList;
+
+#if SHOULD_TRACK_OBJECTS
+	// print out outstanding objects
+	if ((GFrameCounter % 500) == 10)
+	{
+		for (auto It = ClassCounts.CreateIterator(); It; ++It)
+		{
+			NSLog(@"%@ has %d outstanding allocations", It.Key(), It.Value());
+		}
+	}
+#endif
+
+	InitFrame();
+}
+
+
+FMetalTexture2D* FMetalManager::GetBackBuffer()
+{
+	return BackBuffer;
+}
+
+void FMetalManager::PrepareToDraw(uint32 NumVertices)
+{
+#if NO_MEMORY
+	return;
+#endif
+
+	SCOPE_CYCLE_COUNTER(STAT_MetalPrepareDrawTime);
+
+	NumDrawCalls++;
+
+	// make sure the BSS has a valid pipeline state object
+	CurrentBoundShaderState->PrepareToDraw(Pipeline);
+	
+	CommitNonComputeShaderConstants();
+
+	// grow buffers as needed
+	extern TRefCountPtr<FMetalVertexBuffer> HACK_ZeroStrideBuffers[16];
+	for (int32 Index = 0; Index < 16; Index++)
+	{
+		if (HACK_ZeroStrideBuffers[Index] != NULL)
+		{
+			TRefCountPtr<FMetalVertexBuffer> Buffer = HACK_ZeroStrideBuffers[Index];
+			
+			// how big do we need?
+			uint32 NeededSize = Align(NumVertices * Buffer->ZeroStrideElementSize, 128*1024);
+			// if too small, then grow and fill in
+			if (NeededSize > [Buffer->Buffer length])
+			{
+				id<MTLBuffer> Orig = Buffer->Buffer;
+				Buffer->Buffer = [Device newBufferWithLength:NeededSize options:BUFFER_CACHE_MODE];
+				TRACK_OBJECT(Buffer->Buffer);
+
+				uint8* SrcMem = (uint8*)[Orig contents];
+				uint8* DestMem = (uint8*)[Buffer->Buffer contents];
+				for (int32 CopyIndex = 1; CopyIndex < NumVertices; CopyIndex++)
+				{
+					FMemory::Memcpy(DestMem + Buffer->ZeroStrideElementSize * CopyIndex, SrcMem, Buffer->ZeroStrideElementSize);
+				}
+
+				NSLog(@"Expaning buffer from %d to %d", (int32)[Orig length], NeededSize);
+				// set the new buffer object
+				[FMetalManager::GetContext() setVertexBuffer:Buffer->Buffer offset:0 atIndex:UNREAL_TO_METAL_BUFFER_INDEX(Index)];
+
+				ReleaseObject(Orig);
+			}
+		}
+	}
+}
+
+void FMetalManager::CheckStride(uint32 StreamIndex, uint32 Stride)
+{
+	FVertexDeclarationElementList& Elements = CurrentBoundShaderState->VertexDeclaration->Elements;
+	for (int32 ElementIndex = 0; ElementIndex < Elements.Num(); ElementIndex++)
+	{
+		FVertexElement& Element = Elements[ElementIndex];
+		if (Element.StreamIndex == StreamIndex)
+		{
+			check(Element.Stride != 0xFFFF);
+//			checkf(Element.Stride == Stride, TEXT("Strides didn't match up in Stream %d, expected %d, it was %d"), StreamIndex, Stride, Element.Stride);
+		}
+	}
+}
+
+void FMetalManager::SetDepthStencilWriteEnabled(bool bIsDepthWriteEnabled, bool bIsStencilWriteEnabled)
+{
+	Pipeline.bIsDepthWriteEnabled = bIsDepthWriteEnabled;
+	Pipeline.bIsStencilWriteEnabled = bIsStencilWriteEnabled;
+
+	SET_HASH(OFFSET_DEPTH_WRITE_ENABLED, NUMBITS_DEPTH_WRITE_ENABLED, bIsDepthWriteEnabled);
+	SET_HASH(OFFSET_STENCIL_WRITE_ENABLED, NUMBITS_STENCIL_WRITE_ENABLED, bIsStencilWriteEnabled);
+}
+
+void FMetalManager::SetBlendState(FMetalBlendState* BlendState)
+{
+	for(uint32 RenderTargetIndex = 0;RenderTargetIndex < MaxMetalRenderTargets; ++RenderTargetIndex)
+	{
+		MTLBlendDescriptor* Blend0 = BlendState->RenderTargetStates[RenderTargetIndex];
+		check(RenderTargetIndex == 0);
+
+		Pipeline.BlendState[RenderTargetIndex] = Blend0;
+		SET_HASH(OFFSET_SOURCE_RGB_BLEND_FACTOR, NUMBITS_SOURCE_RGB_BLEND_FACTOR, Blend0.sourceRGBBlendFactor);
+		SET_HASH(OFFSET_DEST_RGB_BLEND_FACTOR, NUMBITS_DEST_RGB_BLEND_FACTOR, Blend0.destinationRGBBlendFactor);
+		SET_HASH(OFFSET_RGB_BLEND_OPERATION, NUMBITS_RGB_BLEND_OPERATION, Blend0.rgbBlendOperation);
+		SET_HASH(OFFSET_SOURCE_A_BLEND_FACTOR, NUMBITS_SOURCE_A_BLEND_FACTOR, Blend0.sourceAlphaBlendFactor);
+		SET_HASH(OFFSET_DEST_A_BLEND_FACTOR, NUMBITS_DEST_A_BLEND_FACTOR, Blend0.destinationAlphaBlendFactor);
+		SET_HASH(OFFSET_A_BLEND_OPERATION, NUMBITS_A_BLEND_OPERATION, Blend0.alphaBlendOperation);
+		SET_HASH(OFFSET_WRITE_MASK, NUMBITS_WRITE_MASK, Blend0.writeMask);
+	}
+}
+
+void FMetalManager::SetBoundShaderState(FMetalBoundShaderState* BoundShaderState)
+{
+#if NO_DRAW
+	return;
+#endif
+	CurrentBoundShaderState = BoundShaderState;
+}
+
+void FMetalManager::SetCurrentRenderTarget(FMetalSurface* RenderSurface)
+{
+	// update the current rendered-to pixel format
+	if (RenderSurface)
+	{
+		if (GUseNewDrawable)
+		{
+			// first time in a frame that we are setting the backbuffer, get it
+			if (RenderSurface == &BackBuffer->Surface && RenderSurface->Texture == nil && CurrentDrawable == nil)
+			{
+				SCOPE_CYCLE_COUNTER(STAT_MetalMakeDrawableTime);
+
+				uint32 IdleStart = FPlatformTime::Cycles();
+
+				// make a drawable object for this frame
+				CurrentDrawable = [[IOSAppDelegate GetDelegate].MetalView MakeDrawable];
+				TRACK_OBJECT(CurrentDrawable);
+
+				GRenderThreadIdle[ERenderThreadIdleTypes::WaitingForGPUPresent] += FPlatformTime::Cycles() - IdleStart;
+				GRenderThreadNumIdle[ERenderThreadIdleTypes::WaitingForGPUPresent]++;
+
+				// set the texture into the backbuffer
+				RenderSurface->Texture = CurrentDrawable.texture;
+			}
+		}
+
+		CurrentColorRenderTexture = RenderSurface->Texture;
+		CurrentMSAARenderTexture = RenderSurface->MSAATexture;
+	}
+	else
+	{
+		CurrentColorRenderTexture = nil;
+	}
+}
+
+void FMetalManager::SetCurrentDepthStencilTarget(FMetalSurface* RenderSurface)
+{
+	if (RenderSurface)
+	{
+		// @todo urban: Handle stencil
+		CurrentDepthRenderTexture = RenderSurface->Texture;
+	}
+	else
+	{
+		CurrentDepthRenderTexture = nil;
+	}
+}
+
+void FMetalManager::UpdateContext()
+{
+	if (CurrentColorRenderTexture == PreviousColorRenderTexture && CurrentDepthRenderTexture == PreviousDepthRenderTexture)
+	{
+		return;
+	}
+
+	// if we are setting them to nothing, then this is probably end of frame, and we can't make a framebuffer
+	// with nothng, so just abort this
+	if (CurrentColorRenderTexture == nil && CurrentDepthRenderTexture == nil)
+	{
+		return;
+	}
+
+	MTLAttachmentDescriptor* ColorAttachment = nil;
+	MTLAttachmentDescriptor* DepthAttachment = nil;
+	// @todo urban: Handle stencil
+	MTLAttachmentDescriptor* StencilAttachment = nil;
+
+	Pipeline.SampleCount = 0;
+	if (CurrentColorRenderTexture)
+	{
+		if (CurrentMSAARenderTexture)
+		{
+			ColorAttachment = [MTLAttachmentDescriptor attachmentDescriptorWithTexture:CurrentMSAARenderTexture];
+			[ColorAttachment setLoadAction:MTLLoadActionDontCare];
+			[ColorAttachment setStoreAction:MTLStoreActionMultisampleResolve];
+			[ColorAttachment setResolveTexture:CurrentColorRenderTexture];
+			Pipeline.SampleCount = CurrentMSAARenderTexture.sampleCount;
+		}
+		else
+		{
+			ColorAttachment = [MTLAttachmentDescriptor attachmentDescriptorWithTexture:CurrentColorRenderTexture];
+			[ColorAttachment setLoadAction:MTLLoadActionDontCare];
+			[ColorAttachment setStoreAction:MTLStoreActionStore];
+			Pipeline.SampleCount = 1;
+		}
+
+		Pipeline.RenderTargetFormat[0] = CurrentColorRenderTexture.pixelFormat;
+	}
+	else
+	{
+		Pipeline.RenderTargetFormat[0] = MTLPixelFormatInvalid;
+	}
+	SET_HASH(OFFSET_RENDER_TARGET_FORMAT, NUMBITS_RENDER_TARGET_FORMAT, Pipeline.RenderTargetFormat[0]);
+	
+	if (CurrentDepthRenderTexture)
+	{
+		DepthAttachment = [MTLAttachmentDescriptor attachmentDescriptorWithTexture : CurrentDepthRenderTexture];
+		[DepthAttachment setLoadAction:MTLLoadActionClear];
+		[DepthAttachment setStoreAction:MTLStoreActionDontCare];
+		[DepthAttachment setClearValue:MTLClearValueMakeDepth(0.0f)];
+
+		Pipeline.DepthTargetFormat = CurrentDepthRenderTexture.pixelFormat;
+		if (Pipeline.SampleCount == 0)
+		{
+			Pipeline.SampleCount = CurrentDepthRenderTexture.sampleCount;
+		}
+	}
+	else
+	{
+		Pipeline.DepthTargetFormat = MTLPixelFormatInvalid;
+	}
+	SET_HASH(OFFSET_DEPTH_TARGET_FORMAT, NUMBITS_DEPTH_TARGET_FORMAT, Pipeline.DepthTargetFormat);
+	SET_HASH(OFFSET_SAMPLE_COUNT, NUMBITS_SAMPLE_COUNT, Pipeline.SampleCount);
+
+	PreviousColorRenderTexture = CurrentColorRenderTexture;
+	PreviousDepthRenderTexture = CurrentDepthRenderTexture;
+
+	// commit pending commands on the old render target
+	if (CurrentContext)
+	{
+		if (NumDrawCalls == 0)
+		{
+			NSLog(@"There were %d draw calls for an RT in frame %lld", NumDrawCalls, GFrameCounter);
+		}
+
+		if (true)//NumDrawCalls > 0)
+		{
+			[CurrentContext endEncoding];
+			NumDrawCalls = 0;
+		}
+
+		ReleaseObject(CurrentContext);
+		ReleaseObject(CurrentFramebuffer);
+
+		// if we are doing occlusion queries, we could use this method, along with a completion callback
+		// to set a "render target complete" flag that the OQ code next frame would wait on
+		// commit the buffer for this context
+		[CurrentCommandBuffer commit];
+		UNTRACK_OBJECT(CurrentCommandBuffer);
+		[CurrentCommandBuffer release];
+
+		// create the command buffer for this frame
+		CurrentCommandBuffer = [CommandQueue commandBufferWithUnretainedReferences];
+		[CurrentCommandBuffer retain];
+		TRACK_OBJECT(CurrentCommandBuffer);
+
+#if GPU_TIMINGS
+		uint8 LocalTimerIndex = TimerIndex++;
+#endif
+
+		uint64 LocalCommandBufferIndex = CommandBufferIndex++;
+
+		[CurrentCommandBuffer addCompletedHandler : ^ (id <MTLCommandBuffer> Buffer)
+		{
+#if GPU_TIMINGS
+			// timer stuff
+			GPUTimers[LocalTimerIndex] = FPlatformTime::Seconds();
+#endif
+			FMetalManager::Get()->SetCompletedCommandBufferIndex(LocalCommandBufferIndex);
+		}];
+	}
+		
+	// make a new framebuffer object to render to
+	MTLFramebufferDescriptor* FramebufferDescriptor = [MTLFramebufferDescriptor framebufferDescriptorWithColorAttachment:ColorAttachment depthAttachment:DepthAttachment stencilAttachment:nil];
+	// if we need to do queries, write to the ring buffer (we set the offset into the ring buffer per query)
+	FramebufferDescriptor.visibilityResultBuffer = QueryBuffer.Buffer;
+
+	CurrentFramebuffer = [Device newFramebufferWithDescriptor:FramebufferDescriptor];
+	TRACK_OBJECT(CurrentFramebuffer);
+
+	// make a new render context to use to render to the framebuffer
+	CurrentContext = [CurrentCommandBuffer renderCommandEncoderWithFramebuffer:CurrentFramebuffer];
+//	[CurrentContext setLabel : @"HelloWorld"];
+	[CurrentContext retain];
+	TRACK_OBJECT(CurrentContext);
+
+	// make suire the rasterizer state is set the first time for each new encoder
+	bFirstRasterizerState = true;
+}
+
+
+
+uint32 FRingBuffer::Allocate(uint32 Size, uint32 Alignment)
+{
+	if (Alignment == 0)
+	{
+		Alignment = DefaultAlignment;
+	}
+
+	// align the offset
+	Offset = Align(Offset, Alignment);
+	
+	// wrap if needed
+	if (Offset + Size > Buffer.length)
+	{
+// 		NSLog(@"Wrapping at frame %lld [size = %d]", GFrameCounter, (uint32)Buffer.length);
+		Offset = 0;
+	}
+	
+	// get current location
+	uint32 ReturnOffset = Offset;
+	
+	// allocate
+	Offset += Size;
+	
+	return ReturnOffset;
+}
+
+uint32 FMetalManager::AllocateFromRingBuffer(uint32 Size, uint32 Alignment)
+{
+	return RingBuffer.Allocate(Size, Alignment);
+}
+
+uint32 FMetalManager::AllocateFromQueryBuffer()
+{
+	return QueryBuffer.Allocate(8, 0);
+}
+
+void FMetalManager::CommitNonComputeShaderConstants()
+{
+#if NO_MEMORY
+	return;
+#endif
+	ShaderParameters[METAL_SHADER_STAGE_VERTEX].CommitPackedUniformBuffers(CurrentBoundShaderState, METAL_SHADER_STAGE_VERTEX, CurrentBoundShaderState->VertexShader->BoundUniformBuffers, CurrentBoundShaderState->VertexShader->UniformBuffersCopyInfo);
+	ShaderParameters[METAL_SHADER_STAGE_VERTEX].CommitPackedGlobals(METAL_SHADER_STAGE_VERTEX, CurrentBoundShaderState->VertexShader->Bindings);
+	
+	if (IsValidRef(CurrentBoundShaderState->PixelShader))
+	{
+		ShaderParameters[METAL_SHADER_STAGE_PIXEL].CommitPackedUniformBuffers(CurrentBoundShaderState, METAL_SHADER_STAGE_PIXEL, CurrentBoundShaderState->PixelShader->BoundUniformBuffers, CurrentBoundShaderState->PixelShader->UniformBuffersCopyInfo);
+		ShaderParameters[METAL_SHADER_STAGE_PIXEL].CommitPackedGlobals(METAL_SHADER_STAGE_PIXEL, CurrentBoundShaderState->PixelShader->Bindings);
+	}
+}
+
+bool FMetalManager::WaitForCommandBufferComplete(uint64 IndexToWaitFor, double Timeout)
+{
+	// don't track a block if not needed
+	if (CompletedCommandBufferIndex >= IndexToWaitFor)
+	{
+		return true;
+	}
+
+	// if we don't want to wait, then we have failed
+	if (Timeout == 0.0)
+	{
+		return false;
+	}
+
+	// if we block until it's ready, loop here until it is
+	SCOPE_CYCLE_COUNTER(STAT_RenderQueryResultTime);
+	uint32 IdleStart = FPlatformTime::Cycles();
+	double StartTime = FPlatformTime::Seconds();
+
+	while (CompletedCommandBufferIndex < IndexToWaitFor)
+	{
+		FPlatformProcess::Sleep(0);
+
+		// look for gpu stuck/crashed
+		if ((FPlatformTime::Seconds() - StartTime) > Timeout)
+		{
+			UE_LOG(LogRHI, Log, TEXT("Timed out while waiting for GPU to catch up on occlusion/timer results. (%.1f s)"), Timeout);
+			return false;
+		}
+	}
+
+	// track idle time blocking on GPU
+	GRenderThreadIdle[ERenderThreadIdleTypes::WaitingForGPUQuery] += FPlatformTime::Cycles() - IdleStart;
+	GRenderThreadNumIdle[ERenderThreadIdleTypes::WaitingForGPUQuery]++;
+
+	return true;
+}
+
+
+void FMetalManager::SetRasterizerState(const FRasterizerStateInitializerRHI& State)
+{
+	if (bFirstRasterizerState)
+	{
+		[CurrentContext setFrontFacingWinding:MTLWindingCounterClockwise];
+	}
+
+	if (bFirstRasterizerState || ShadowRasterizerState.CullMode != State.CullMode)
+	{
+		[CurrentContext setCullMode:TranslateCullMode(State.CullMode)];
+		ShadowRasterizerState.CullMode = State.CullMode;
+	}
+
+	// @todo urban: Shoud the clamp be FLT_MAX or 0 for "don't bother"
+	if (bFirstRasterizerState || ShadowRasterizerState.DepthBias != State.DepthBias || ShadowRasterizerState.SlopeScaleDepthBias != State.SlopeScaleDepthBias)
+	{
+		[CurrentContext setDepthBias:State.DepthBias slopeScale:State.SlopeScaleDepthBias clamp : FLT_MAX];
+		ShadowRasterizerState.DepthBias = State.DepthBias;
+		ShadowRasterizerState.SlopeScaleDepthBias = State.SlopeScaleDepthBias;
+	}
+
+	// @todo urban: Would we ever need this in a shipping app?
+#if !UE_BUILD_SHIPPING
+	if (bFirstRasterizerState || ShadowRasterizerState.FillMode != State.FillMode)
+	{
+		[CurrentContext setTriangleFillMode:TranslateFillMode(State.FillMode)];
+		ShadowRasterizerState.FillMode = State.FillMode;
+	}
+#endif
+
+	bFirstRasterizerState = false;
+}
\ No newline at end of file
diff --git a/Engine/Source/Runtime/IOS/MetalRHI/Private/MetalManager.h b/Engine/Source/Runtime/IOS/MetalRHI/Private/MetalManager.h
new file mode 100644
--- /dev/null
+++ b/Engine/Source/Runtime/IOS/MetalRHI/Private/MetalManager.h
@@ -0,0 +1,212 @@
+// Copyright 1998-2014 Epic Games, Inc. All Rights Reserved.
+
+
+#pragma once
+
+#import <Metal/Metal.h>
+#import <QuartzCore/CAMetalLayer.h>
+#include "BLView.h"
+
+
+struct FPipelineShadow
+{
+	FPipelineShadow()
+	{
+		// zero out our memory
+		FMemory::Memzero(this, sizeof(*this));
+	}
+
+	/**
+	 * @return the hash of the current pipeline state
+	 */
+	uint64 GetHash() const
+	{
+		return Hash;
+	}
+
+	/**
+	 * FOrce a pipeline state for a known hash
+	 */
+	void SetHash(uint64 InHash);
+	
+	/**
+	 * @return an pipeline state object that matches the current state and the BSS
+	 */
+	id<MTLRenderPipelineState> CreatePipelineStateForBoundShaderState(FMetalBoundShaderState* BSS) const;
+
+	bool bIsDepthWriteEnabled;
+	bool bIsStencilWriteEnabled;
+	MTLBlendDescriptor* BlendState[MaxMetalRenderTargets];
+	MTLPixelFormat RenderTargetFormat[MaxMetalRenderTargets];
+	MTLPixelFormat DepthTargetFormat;
+    uint32 SampleCount;
+
+	// running hash of the pipeline state
+	uint64 Hash;
+};
+
+struct FRingBuffer
+{
+	FRingBuffer(id<MTLDevice> Device, uint32 Size, uint32 InDefaultAlignment)
+	{
+		DefaultAlignment = InDefaultAlignment;
+		Buffer = [Device newBufferWithLength:Size options:BUFFER_CACHE_MODE];
+		Offset = 0;
+	}
+
+	uint32 Allocate(uint32 Size, uint32 Alignment);
+
+	uint32 DefaultAlignment;
+	id<MTLBuffer> Buffer;
+	uint32 Offset;
+};
+
+class FMetalManager
+{
+public:
+	static FMetalManager* Get();
+	static id<MTLDevice> GetDevice();
+	static id<MTLRenderCommandEncoder> GetContext();
+	static void ReleaseObject(id Object);
+	
+	void BeginFrame();
+	void EndFrame(bool bPreset);
+	
+	/**
+	 * Return a texture object that wraps the backbuffer
+	 */
+	FMetalTexture2D* GetBackBuffer();
+
+//	/**
+//	 * Sets a pending stream source, which will be hooked up to the vertex decl in PrepareToDraw
+//	 */
+//	void SetPendingStream(uint32 StreamIndex, TRefCountPtr<FMetalVertexBuffer> VertexBuffer, uint32 Offset, uint32 Stride);
+
+	/**
+	 * Do anything necessary to prepare for any kind of draw call 
+	 */
+	void PrepareToDraw(uint32 NumVertices);
+	
+	/**
+	 * Check the stride against what we expected
+	 */
+	void CheckStride(uint32 StreamIndex, uint32 Stride);
+	
+	/**
+	 * Functions to update the pipeline descriptor and/or context
+	 */
+	void SetDepthStencilWriteEnabled(bool bIsDepthWriteEnabled, bool bIsStencilWriteEnabled);
+	void SetBlendState(FMetalBlendState* BlendState);
+	void SetBoundShaderState(FMetalBoundShaderState* BoundShaderState);
+	void SetCurrentRenderTarget(FMetalSurface* RenderSurface);
+	void SetCurrentDepthStencilTarget(FMetalSurface* RenderSurface);
+	
+	/**
+	 * Update the context with the current render targets
+	 */
+	void UpdateContext();
+	
+	/**
+	 * Allocate from a dynamic ring buffer - by default align to the allowed alignment for offset field when setting buffers
+	 */
+	uint32 AllocateFromRingBuffer(uint32 Size, uint32 Alignment=0);
+	id<MTLBuffer> GetRingBuffer()
+	{
+		return RingBuffer.Buffer;
+	}
+
+	uint32 AllocateFromQueryBuffer();
+	id<MTLBuffer> GetQueryBuffer()
+	{
+		return QueryBuffer.Buffer;
+	}
+
+	uint64 GetCommandBufferIndex()
+	{
+		return CommandBufferIndex;
+	}
+	
+	void SetCompletedCommandBufferIndex(uint64 Index)
+	{
+		CompletedCommandBufferIndex = Index;
+	}
+
+	bool WaitForCommandBufferComplete(uint64 IndexToWaitFor, double Timeout);
+
+	void SetRasterizerState(const FRasterizerStateInitializerRHI& State);
+
+protected:
+	FMetalManager();
+	void InitFrame();
+	void GenerateFetchShader();
+
+	id<MTLDevice> Device;
+
+	id<MTLCommandQueue> CommandQueue;
+	
+	uint32 NumDrawCalls;
+	id<MTLCommandBuffer> CurrentCommandBuffer;
+	id<CAMetalDrawable> CurrentDrawable;
+
+	dispatch_semaphore_t CommandBufferSemaphore;
+
+	id<MTLFramebuffer> CurrentFramebuffer;
+	id<MTLRenderCommandEncoder> CurrentContext;
+	
+	id<MTLTexture> CurrentColorRenderTexture, PreviousColorRenderTexture;
+	id<MTLTexture> CurrentDepthRenderTexture, PreviousDepthRenderTexture;
+	id<MTLTexture> CurrentMSAARenderTexture;
+	
+	TRefCountPtr<FMetalTexture2D> BackBuffer;
+	TRefCountPtr<FMetalBoundShaderState> CurrentBoundShaderState;
+	
+	// the running pipeline state descriptor object
+	FPipelineShadow Pipeline;
+
+	MTLRenderPipelineDescriptor* CurrentPipelineDescriptor;
+	// a parellel descriptor for no depth buffer (since we can't reset the depth format to invalid YET)
+	MTLRenderPipelineDescriptor* CurrentPipelineDescriptor2;
+	
+
+	// the amazing ring buffer for dynamic data
+	FRingBuffer RingBuffer, QueryBuffer;
+
+	TArray<id> DelayedFreeLists[4];
+	uint32 WhichFreeList;
+
+	// the slot to store a per-thread autorelease pool
+	uint32 AutoReleasePoolTLSSlot;
+
+	// always growing command buffer index
+	uint64 CommandBufferIndex;
+	uint64 CompletedCommandBufferIndex;
+
+	// shadow the rasterizer state to reduce some render encoder pressure
+	FRasterizerStateInitializerRHI ShadowRasterizerState;
+	bool bFirstRasterizerState;
+
+#if GPU_TIMINGS
+	// holds on to timing info
+	double GPUTimers[256];
+	uint8 TimerIndex;
+	uint8 FirstIndexThisFrame;
+	uint8 FirstIndexToDump;
+	uint8 LastIndexToDump;
+#endif
+
+public:
+	void CommitNonComputeShaderConstants();
+	
+	class FMetalShaderParameterCache*	ShaderParameters;
+};
+
+// Stats
+DECLARE_CYCLE_STAT_EXTERN(TEXT("MakeDrawable time"),STAT_MetalMakeDrawableTime,STATGROUP_MetalRHI, );
+DECLARE_CYCLE_STAT_EXTERN(TEXT("Draw call time"),STAT_MetalDrawCallTime,STATGROUP_MetalRHI, );
+DECLARE_CYCLE_STAT_EXTERN(TEXT("PrepareDraw time"),STAT_MetalPrepareDrawTime,STATGROUP_MetalRHI, );
+DECLARE_CYCLE_STAT_EXTERN(TEXT("PipelineState time"),STAT_MetalPipelineStateTime,STATGROUP_MetalRHI, );
+DECLARE_CYCLE_STAT_EXTERN(TEXT("BpoundShaderState time"),STAT_MetalBoundShaderStateTime,STATGROUP_MetalRHI, );
+DECLARE_CYCLE_STAT_EXTERN(TEXT("VertexDeclaration time"),STAT_MetalVertexDeclarationTime,STATGROUP_MetalRHI, );
+DECLARE_CYCLE_STAT_EXTERN(TEXT("Uniform buffer pool cleanup time"), STAT_MetalUniformBufferCleanupTime, STATGROUP_MetalRHI, );
+DECLARE_MEMORY_STAT_EXTERN(TEXT("Uniform buffer pool memory"), STAT_MetalFreeUniformBufferMemory, STATGROUP_MetalRHI, );
+DECLARE_DWORD_ACCUMULATOR_STAT_EXTERN(TEXT("Uniform buffer pool num free"), STAT_MetalNumFreeUniformBuffers, STATGROUP_MetalRHI, );
diff --git a/Engine/Source/Runtime/IOS/MetalRHI/Private/MetalQuery.cpp b/Engine/Source/Runtime/IOS/MetalRHI/Private/MetalQuery.cpp
new file mode 100644
--- /dev/null
+++ b/Engine/Source/Runtime/IOS/MetalRHI/Private/MetalQuery.cpp
@@ -0,0 +1,76 @@
+// Copyright 1998-2014 Epic Games, Inc. All Rights Reserved.
+
+/*=============================================================================
+	MetalQuery.cpp: Metal query RHI implementation.
+=============================================================================*/
+
+#include "MetalRHIPrivate.h"
+
+FMetalRenderQuery::FMetalRenderQuery(ERenderQueryType InQueryType)
+{
+
+}
+
+FMetalRenderQuery::~FMetalRenderQuery()
+{
+
+}
+
+
+void FMetalRenderQuery::Begin()
+{
+	// these only need to be 8 byte aligned (of course, a normal allocation oafter this will be 256 bytes, making large holes if we don't do all 
+	// query allocations at once)
+	Offset = FMetalManager::Get()->AllocateFromQueryBuffer();
+
+	// reset the query to 0
+	uint32* QueryMemory = (uint32*)((uint8*)[FMetalManager::Get()->GetQueryBuffer() contents] + Offset);
+	*QueryMemory = 0;
+
+	// remember which command buffer we are in
+	CommandBufferIndex = FMetalManager::Get()->GetCommandBufferIndex();
+
+	[FMetalManager::GetContext() setVisibilityResultMode:MTLVisibilityResultModeBoolean offset:Offset];
+}
+
+void FMetalRenderQuery::End()
+{
+	// switch back to non-occlusion rendering
+	[FMetalManager::GetContext() setVisibilityResultMode:MTLVisibilityResultModeDisabled offset:0];
+}
+
+
+
+FRenderQueryRHIRef FMetalDynamicRHI::RHICreateRenderQuery(ERenderQueryType QueryType)
+{
+	return new FMetalRenderQuery(QueryType);
+}
+
+void FMetalDynamicRHI::RHIResetRenderQuery(FRenderQueryRHIParamRef QueryRHI)
+{
+	DYNAMIC_CAST_METGALRESOURCE(RenderQuery,Query);
+
+//	NOT_SUPPORTED("RHIResetRenderQuery");
+}
+
+bool FMetalDynamicRHI::RHIGetRenderQueryResult(FRenderQueryRHIParamRef QueryRHI,uint64& OutNumPixels,bool bWait)
+{
+	DYNAMIC_CAST_METGALRESOURCE(RenderQuery, Query);
+
+	uint32* QueryMemory = (uint32*)((uint8*)[FMetalManager::Get()->GetQueryBuffer() contents] + Query->Offset);
+
+	// timer queries are used for Benchmarks which can stall a bit more
+	double Timeout = 0.0;
+	if (bWait)
+	{
+		Timeout = 0.5f;//(Query->QueryType == RQT_AbsoluteTime) ? 2.0 : 0.5;
+	}
+	if (!FMetalManager::Get()->WaitForCommandBufferComplete(Query->CommandBufferIndex, Timeout))
+	{
+		return false;
+	}
+
+	// at this point, we are ready to read the value!
+	OutNumPixels = *QueryMemory;
+	return true;
+}
diff --git a/Engine/Source/Runtime/IOS/MetalRHI/Private/MetalRHI.cpp b/Engine/Source/Runtime/IOS/MetalRHI/Private/MetalRHI.cpp
new file mode 100644
--- /dev/null
+++ b/Engine/Source/Runtime/IOS/MetalRHI/Private/MetalRHI.cpp
@@ -0,0 +1,221 @@
+// Copyright 1998-2014 Epic Games, Inc. All Rights Reserved.
+
+/*=============================================================================
+	MetalRHI.cpp: Metal device RHI implementation.
+=============================================================================*/
+
+#include "MetalRHIPrivate.h"
+#include "IOSAppDelegate.h"
+
+/** Set to 1 to enable GPU events in Xcode frame debugger */
+#define ENABLE_METAL_GPUEVENTS	(UE_BUILD_DEBUG || UE_BUILD_DEVELOPMENT)
+
+DEFINE_LOG_CATEGORY(LogMetal)
+
+bool FMetalDynamicRHIModule::IsSupported()
+{
+	return true;
+}
+
+FDynamicRHI* FMetalDynamicRHIModule::CreateRHI()
+{
+	return new FMetalDynamicRHI();
+}
+
+IMPLEMENT_MODULE(FMetalDynamicRHIModule, MetalRHI);
+
+FMetalDynamicRHI::FMetalDynamicRHI()
+{
+	// This should be called once at the start 
+	check( IsInGameThread() );
+	check( !GIsThreadedRendering );
+
+	// Initialize the RHI capabilities.
+	GRHIFeatureLevel = ERHIFeatureLevel::ES2;
+	GRHIShaderPlatform = SP_METAL;
+
+	GPixelCenterOffset = 0.0f;
+	GSupportsVertexInstancing = false;
+	GSupportsEmulatedVertexInstancing = false;
+	GSupportsVertexTextureFetch = false;
+	GSupportsShaderFramebufferFetch = true;
+	GHardwareHiddenSurfaceRemoval = true;
+//	GRHIAdapterName = 
+//	GRHIVendorId = 
+	GSupportsRenderTargetFormat_PF_G8 = false;
+	GSupportsQuads = false;
+	GRHISupportsTextureStreaming = false;
+	GMaxShadowDepthBufferSizeX = 4096;
+	GMaxShadowDepthBufferSizeY = 4096;
+// 	GReadTexturePoolSizeFromIni = true;
+ 
+	GMaxTextureDimensions = 4096;
+	GMaxTextureMipCount = FPlatformMath::CeilLogTwo( GMaxTextureDimensions ) + 1;
+	GMaxTextureMipCount = FPlatformMath::Min<int32>( MAX_TEXTURE_MIP_COUNT, GMaxTextureMipCount );
+	GMaxCubeTextureDimensions = 4096;
+	GMaxTextureArrayLayers = 2048;
+
+	GShaderPlatformForFeatureLevel[ERHIFeatureLevel::ES2] = SP_METAL;
+	GShaderPlatformForFeatureLevel[ERHIFeatureLevel::SM3] = SP_NumPlatforms;
+	GShaderPlatformForFeatureLevel[ERHIFeatureLevel::SM4] = SP_NumPlatforms;
+	GShaderPlatformForFeatureLevel[ERHIFeatureLevel::SM5] = SP_NumPlatforms;
+
+// 	GDrawUPVertexCheckCount = MAX_uint16;
+
+	// Initialize the platform pixel format map.
+	GPixelFormats[PF_Unknown			].PlatformFormat	= MTLPixelFormatInvalid;
+	GPixelFormats[PF_A32B32G32R32F		].PlatformFormat	= MTLPixelFormatRGBA32Float;
+	GPixelFormats[PF_B8G8R8A8			].PlatformFormat	= MTLPixelFormatBGRA8Unorm;
+	GPixelFormats[PF_G8					].PlatformFormat	= MTLPixelFormatR8Unorm;
+	GPixelFormats[PF_G16				].PlatformFormat	= MTLPixelFormatR16Unorm;
+	GPixelFormats[PF_DXT1				].PlatformFormat	= MTLPixelFormatInvalid;
+	GPixelFormats[PF_DXT3				].PlatformFormat	= MTLPixelFormatInvalid;
+	GPixelFormats[PF_DXT5				].PlatformFormat	= MTLPixelFormatInvalid;
+	GPixelFormats[PF_PVRTC2				].PlatformFormat	= MTLPixelFormatPVRTC_RGBA_2BPP;
+	GPixelFormats[PF_PVRTC2				].Supported			= true;
+	GPixelFormats[PF_PVRTC4				].PlatformFormat	= MTLPixelFormatPVRTC_RGBA_4BPP;
+	GPixelFormats[PF_PVRTC4				].Supported			= true;
+	GPixelFormats[PF_UYVY				].PlatformFormat	= MTLPixelFormatInvalid;
+	GPixelFormats[PF_FloatRGB			].PlatformFormat	= MTLPixelFormatRGBA16Float;
+	GPixelFormats[PF_FloatRGB			].BlockBytes		= 8;
+	GPixelFormats[PF_FloatRGBA			].PlatformFormat	= MTLPixelFormatRGBA16Float;
+	GPixelFormats[PF_FloatRGBA			].BlockBytes		= 8;
+	GPixelFormats[PF_DepthStencil		].PlatformFormat	= MTLPixelFormatDepth32Float;
+	GPixelFormats[PF_DepthStencil		].BlockBytes		= 4;
+	GPixelFormats[PF_X24_G8				].PlatformFormat	= MTLPixelFormatStencil8;
+	GPixelFormats[PF_X24_G8				].BlockBytes		= 1;
+	GPixelFormats[PF_ShadowDepth		].PlatformFormat	= MTLPixelFormatDepth32Float;
+	GPixelFormats[PF_R32_FLOAT			].PlatformFormat	= MTLPixelFormatR32Float;
+	GPixelFormats[PF_G16R16				].PlatformFormat	= MTLPixelFormatRG16Unorm;
+	GPixelFormats[PF_G16R16F			].PlatformFormat	= MTLPixelFormatRG16Float;
+	GPixelFormats[PF_G16R16F_FILTER		].PlatformFormat	= MTLPixelFormatRG16Float;
+	GPixelFormats[PF_G32R32F			].PlatformFormat	= MTLPixelFormatRG32Float;
+	GPixelFormats[PF_A2B10G10R10		].PlatformFormat    = MTLPixelFormatA2BGR10Unorm;
+	GPixelFormats[PF_A16B16G16R16		].PlatformFormat    = MTLPixelFormatRGBA16Float;
+	GPixelFormats[PF_D24				].PlatformFormat	= MTLPixelFormatInvalid;
+	GPixelFormats[PF_R16F				].PlatformFormat	= MTLPixelFormatR16Float;
+	GPixelFormats[PF_R16F_FILTER		].PlatformFormat	= MTLPixelFormatR16Float;
+	GPixelFormats[PF_BC5				].PlatformFormat	= MTLPixelFormatInvalid;
+	GPixelFormats[PF_V8U8				].PlatformFormat	= 
+	GPixelFormats[PF_A1					].PlatformFormat	= MTLPixelFormatInvalid;
+	GPixelFormats[PF_FloatR11G11B10		].PlatformFormat	= MTLPixelFormatB10GR11Float;
+	GPixelFormats[PF_FloatR11G11B10		].BlockBytes		= 4;
+	GPixelFormats[PF_A8					].PlatformFormat	= MTLPixelFormatR8Unorm;
+	GPixelFormats[PF_R32_UINT			].PlatformFormat	= MTLPixelFormatR32Uint;
+	GPixelFormats[PF_R32_SINT			].PlatformFormat	= MTLPixelFormatR32Sint;
+	GPixelFormats[PF_R16G16B16A16_UINT	].PlatformFormat	= MTLPixelFormatRGBA16Uint;
+	GPixelFormats[PF_R16G16B16A16_SINT	].PlatformFormat	= MTLPixelFormatRGBA16Sint;
+	GPixelFormats[PF_R5G6B5_UNORM		].PlatformFormat	= MTLPixelFormatRGB565Unorm;
+	GPixelFormats[PF_R8G8B8A8			].PlatformFormat	= MTLPixelFormatRGBA8Unorm;
+	GPixelFormats[PF_R8G8				].PlatformFormat	= MTLPixelFormatRG8Unorm;
+
+	GDynamicRHI = this;
+	GIsRHIInitialized = true;
+
+	// Notify all initialized FRenderResources that there's a valid RHI device to create their RHI resources for now.
+	for(TLinkedList<FRenderResource*>::TIterator ResourceIt(FRenderResource::GetResourceList());ResourceIt;ResourceIt.Next())
+	{
+		ResourceIt->InitDynamicRHI();
+	}
+	for(TLinkedList<FRenderResource*>::TIterator ResourceIt(FRenderResource::GetResourceList());ResourceIt;ResourceIt.Next())
+	{
+		ResourceIt->InitRHI();
+	}
+}
+
+FMetalDynamicRHI::~FMetalDynamicRHI()
+{
+	check(IsInGameThread() && IsInRenderingThread());
+
+}
+
+uint64 FMetalDynamicRHI::RHICalcTexture2DPlatformSize(uint32 SizeX, uint32 SizeY, uint8 Format, uint32 NumMips, uint32 NumSamples, uint32 Flags, uint32& OutAlign)
+{
+	OutAlign = 0;
+	return CalcTextureSize(SizeX, SizeY, (EPixelFormat)Format, NumMips);
+}
+
+uint64 FMetalDynamicRHI::RHICalcTexture3DPlatformSize(uint32 SizeX, uint32 SizeY, uint32 SizeZ, uint8 Format, uint32 NumMips, uint32 Flags, uint32& OutAlign)
+{
+	OutAlign = 0;
+	return CalcTextureSize3D(SizeX, SizeY, SizeZ, (EPixelFormat)Format, NumMips);
+}
+
+uint64 FMetalDynamicRHI::RHICalcTextureCubePlatformSize(uint32 Size, uint8 Format, uint32 NumMips, uint32 Flags,	uint32& OutAlign)
+{
+	OutAlign = 0;
+	return CalcTextureSize(Size, Size, (EPixelFormat)Format, NumMips) * 6;
+}
+
+
+void FMetalDynamicRHI::Init()
+{
+
+}
+
+void FMetalDynamicRHI::RHIBeginFrame()
+{
+}
+
+void FMetalDynamicRHI::RHIEndFrame()
+{
+}
+
+void FMetalDynamicRHI::RHIBeginScene()
+{
+
+}
+
+void FMetalDynamicRHI::RHIEndScene()
+{
+
+}
+
+void FMetalDynamicRHI::PushEvent(const TCHAR* Name)
+{
+#if ENABLE_METAL_GPUEVENTS
+	[FMetalManager::GetContext() pushDebugGroup: [NSString stringWithTCHARString:Name]];
+#endif
+}
+
+void FMetalDynamicRHI::PopEvent()
+{
+#if ENABLE_METAL_GPUEVENTS
+	[FMetalManager::GetContext() popDebugGroup];
+#endif
+}
+
+void FMetalDynamicRHI::RHIGetSupportedResolution( uint32 &Width, uint32 &Height )
+{
+
+}
+
+bool FMetalDynamicRHI::RHIGetAvailableResolutions(FScreenResolutionArray& Resolutions, bool bIgnoreRefreshRate)
+{
+
+	return false;
+}
+
+void FMetalDynamicRHI::RHIFlushResources()
+{
+
+}
+
+void FMetalDynamicRHI::RHIAcquireThreadOwnership()
+{
+
+}
+void FMetalDynamicRHI::RHIReleaseThreadOwnership()
+{
+
+}
+
+void FMetalDynamicRHI::RHIGpuTimeBegin(uint32 Hash, bool bCompute)
+{
+
+}
+
+void FMetalDynamicRHI::RHIGpuTimeEnd(uint32 Hash, bool bCompute)
+{
+
+}
diff --git a/Engine/Source/Runtime/IOS/MetalRHI/Private/MetalRHIPrivate.h b/Engine/Source/Runtime/IOS/MetalRHI/Private/MetalRHIPrivate.h
new file mode 100644
--- /dev/null
+++ b/Engine/Source/Runtime/IOS/MetalRHI/Private/MetalRHIPrivate.h
@@ -0,0 +1,266 @@
+// Copyright 1998-2014 Epic Games, Inc. All Rights Reserved.
+
+/*=============================================================================
+	MetalRHIPrivate.h: Private Metal RHI definitions.
+=============================================================================*/
+
+#pragma once
+
+#include "Engine.h"
+
+// UE4 has a Max of 8 RTs, but we max out at 4, no need to waste time processing 8
+// @todo urban: Currently only using 1, so just do 1!
+const uint32 MaxMetalRenderTargets = 1;
+
+// How many possible vertex streams are allowed
+const uint32 MaxMetalStreams = 16;
+
+// Requirement for vertex buffer offset field
+const uint32 BufferOffsetAlignment = 256;
+
+//#define BUFFER_CACHE_MODE MTLResourceOptionCPUCacheModeDefault
+#define BUFFER_CACHE_MODE MTLResourceOptionCPUCacheModeWriteCombined
+
+#define GPU_TIMINGS 0 // (UE_BUILD_DEBUG || UE_BUILD_DEVELOPMENT)
+
+#define SHOULD_TRACK_OBJECTS 0 // (UE_BUILD_DEBUG || UE_BUILD_DEVELOPMENT)
+
+#define UNREAL_TO_METAL_BUFFER_INDEX(Index) (15 - Index)
+
+// Dependencies
+#include "MetalRHI.h"
+#include "MetalGlobalUniformBuffer.h"
+#include "RHI.h"
+#include "MetalManager.h"
+
+
+// Access the underlying surface object from any kind of texture
+FMetalSurface& GetMetalSurfaceFromRHITexture(FRHITexture* Texture);
+
+#define NOT_SUPPORTED(Func) UE_LOG(LogMetal, Fatal, TEXT("'%s' is not supported"), L##Func);
+
+#if SHOULD_TRACK_OBJECTS
+extern TMap<id, int32> ClassCounts;
+#define TRACK_OBJECT(Obj) ClassCounts.FindOrAdd([Obj class])++;
+#define UNTRACK_OBJECT(Obj) ClassCounts.FindOrAdd([Obj class])--;
+#else
+#define TRACK_OBJECT(Obj)
+#define UNTRACK_OBJECT(Obj)
+#endif
+
+
+
+
+
+
+class FLocalEvent : public FEvent
+{
+	// This is a little complicated, in an attempt to match Win32 Event semantics...
+    typedef enum
+    {
+        TRIGGERED_NONE,
+        TRIGGERED_ONE,
+        TRIGGERED_ALL,
+    } TriggerType;
+	
+	bool bInitialized;
+	bool bIsManualReset;
+	volatile TriggerType Triggered;
+	volatile int32 WaitingThreads;
+    pthread_mutex_t Mutex;
+    pthread_cond_t Condition;
+	
+	inline void LockEventMutex()
+	{
+		const int rc = pthread_mutex_lock(&Mutex);
+		check(rc == 0);
+	}
+	
+	inline void UnlockEventMutex()
+	{
+		const int rc = pthread_mutex_unlock(&Mutex);
+		check(rc == 0);
+	}
+	
+	static inline void SubtractTimevals(const struct timeval *FromThis, struct timeval *SubThis, struct timeval *Difference)
+	{
+		if (FromThis->tv_usec < SubThis->tv_usec)
+		{
+			int nsec = (SubThis->tv_usec - FromThis->tv_usec) / 1000000 + 1;
+			SubThis->tv_usec -= 1000000 * nsec;
+			SubThis->tv_sec += nsec;
+		}
+		
+		if (FromThis->tv_usec - SubThis->tv_usec > 1000000)
+		{
+			int nsec = (FromThis->tv_usec - SubThis->tv_usec) / 1000000;
+			SubThis->tv_usec += 1000000 * nsec;
+			SubThis->tv_sec -= nsec;
+		}
+		
+		Difference->tv_sec = FromThis->tv_sec - SubThis->tv_sec;
+		Difference->tv_usec = FromThis->tv_usec - SubThis->tv_usec;
+	}
+	
+public:
+	
+	FLocalEvent()
+	{
+		bInitialized = false;
+		bIsManualReset = false;
+		Triggered = TRIGGERED_NONE;
+		WaitingThreads = 0;
+	}
+	
+	virtual ~FLocalEvent()
+	{
+		// Safely destructing an Event is VERY delicate, so it can handle badly-designed
+		//  calling code that might still be waiting on the event.
+		if (bInitialized)
+		{
+			// try to flush out any waiting threads...
+			LockEventMutex();
+			bIsManualReset = true;
+			UnlockEventMutex();
+			Trigger();  // any waiting threads should start to wake up now.
+			
+			LockEventMutex();
+			bInitialized = false;  // further incoming calls to this object will now crash in check().
+			while (WaitingThreads)  // spin while waiting threads wake up.
+			{
+				UnlockEventMutex();  // cycle through waiting threads...
+				LockEventMutex();
+			}
+			// No threads are currently waiting on Condition and we hold the Mutex. Kill it.
+			pthread_cond_destroy(&Condition);
+			// Unlock and kill the mutex, since nothing else can grab it now.
+			UnlockEventMutex();
+			pthread_mutex_destroy(&Mutex);
+		}
+	}
+	
+	virtual bool Create(bool _bIsManualReset = false) OVERRIDE
+	{
+		check(!bInitialized);
+		bool RetVal = false;
+		Triggered = TRIGGERED_NONE;
+		bIsManualReset = _bIsManualReset;
+		
+		if (pthread_mutex_init(&Mutex, NULL) == 0)
+		{
+			if (pthread_cond_init(&Condition, NULL) == 0)
+			{
+				bInitialized = true;
+				RetVal = true;
+			}
+			else
+			{
+				pthread_mutex_destroy(&Mutex);
+			}
+		}
+		return RetVal;
+	}
+	
+	virtual void Trigger() OVERRIDE
+	{
+		check(bInitialized);
+		
+		LockEventMutex();
+		
+		if (bIsManualReset)
+		{
+			// Release all waiting threads at once.
+			Triggered = TRIGGERED_ALL;
+			int rc = pthread_cond_broadcast(&Condition);
+			check(rc == 0);
+		}
+		else
+		{
+			// Release one or more waiting threads (first one to get the mutex
+			//  will reset Triggered, rest will go back to waiting again).
+			Triggered = TRIGGERED_ONE;
+			int rc = pthread_cond_signal(&Condition);  // may release multiple threads anyhow!
+			check(rc == 0);
+		}
+		
+		UnlockEventMutex();
+	}
+	
+	virtual void Reset() OVERRIDE
+	{
+		check(bInitialized);
+		LockEventMutex();
+		Triggered = TRIGGERED_NONE;
+		UnlockEventMutex();
+	}
+	
+	virtual bool Wait(uint32 WaitTime)
+	{
+		check(bInitialized);
+		
+		struct timeval StartTime;
+		
+		// We need to know the start time if we're going to do a timed wait.
+		if ( (WaitTime > 0) && (WaitTime != ((uint32)-1)) )  // not polling and not infinite wait.
+		{
+			gettimeofday(&StartTime, NULL);
+		}
+		
+		LockEventMutex();
+		
+		bool bRetVal = false;
+		
+		// loop in case we fall through the Condition signal but someone else claims the event.
+		do
+		{
+			// See what state the event is in...we may not have to wait at all...
+			
+			// One thread should be released. We saw it first, so we'll take it.
+			if (Triggered == TRIGGERED_ONE)
+			{
+				Triggered = TRIGGERED_NONE;  // dibs!
+				bRetVal = true;
+			}
+			
+			// manual reset that is still signaled. Every thread goes.
+			else if (Triggered == TRIGGERED_ALL)
+			{
+				bRetVal = true;
+			}
+			
+			// No event signalled yet.
+			else if (WaitTime != 0)  // not just polling, wait on the condition variable.
+			{
+				WaitingThreads++;
+				if (WaitTime == ((uint32)-1)) // infinite wait?
+				{
+					int rc = pthread_cond_wait(&Condition, &Mutex);  // unlocks Mutex while blocking...
+					check(rc == 0);
+				}
+				else  // timed wait.
+				{
+					struct timespec TimeOut;
+					const uint32 ms = (StartTime.tv_usec / 1000) + WaitTime;
+					TimeOut.tv_sec = StartTime.tv_sec + (ms / 1000);
+					TimeOut.tv_nsec = (ms % 1000) * 1000000;  // remainder of milliseconds converted to nanoseconds.
+					int rc = pthread_cond_timedwait(&Condition, &Mutex, &TimeOut);    // unlocks Mutex while blocking...
+					check((rc == 0) || (rc == ETIMEDOUT));
+					
+					// Update WaitTime and StartTime in case we have to go again...
+					struct timeval Now, Difference;
+					gettimeofday(&Now, NULL);
+					SubtractTimevals(&Now, &StartTime, &Difference);
+					const int32 DifferenceMS = ((Difference.tv_sec * 1000) + (Difference.tv_usec / 1000));
+					WaitTime = ((DifferenceMS >= WaitTime) ? 0 : (WaitTime - DifferenceMS));
+					StartTime = Now;
+				}
+				WaitingThreads--;
+				check(WaitingThreads >= 0);
+			}
+			
+		} while ((!bRetVal) && (WaitTime != 0));
+		
+		UnlockEventMutex();
+		return bRetVal;
+	}
+};
diff --git a/Engine/Source/Runtime/IOS/MetalRHI/Private/MetalRenderTarget.cpp b/Engine/Source/Runtime/IOS/MetalRHI/Private/MetalRenderTarget.cpp
new file mode 100644
--- /dev/null
+++ b/Engine/Source/Runtime/IOS/MetalRHI/Private/MetalRenderTarget.cpp
@@ -0,0 +1,39 @@
+// Copyright 1998-2014 Epic Games, Inc. All Rights Reserved.
+
+/*=============================================================================
+	MetalRenderTarget.cpp: Metal render target implementation.
+=============================================================================*/
+
+#include "MetalRHIPrivate.h"
+#include "ScreenRendering.h"
+
+
+void FMetalDynamicRHI::RHICopyToResolveTarget(FTextureRHIParamRef SourceTextureRHI, FTextureRHIParamRef DestTextureRHI, bool bKeepOriginalSurface, const FResolveParams& ResolveParams)
+{
+
+}
+
+void FMetalDynamicRHI::RHIReadSurfaceData(FTextureRHIParamRef TextureRHI, FIntRect Rect, TArray<FColor>& OutData, FReadSurfaceDataFlags InFlags)
+{
+
+}
+
+void FMetalDynamicRHI::RHIMapStagingSurface(FTextureRHIParamRef TextureRHI,void*& OutData,int32& OutWidth,int32& OutHeight)
+{
+
+}
+
+void FMetalDynamicRHI::RHIUnmapStagingSurface(FTextureRHIParamRef TextureRHI)
+{
+
+}
+
+void FMetalDynamicRHI::RHIReadSurfaceFloatData(FTextureRHIParamRef TextureRHI, FIntRect Rect, TArray<FFloat16Color>& OutData, ECubeFace CubeFace,int32 ArrayIndex,int32 MipIndex)
+{
+
+}
+
+void FMetalDynamicRHI::RHIRead3DSurfaceFloatData(FTextureRHIParamRef TextureRHI,FIntRect InRect,FIntPoint ZMinMax,TArray<FFloat16Color>& OutData)
+{
+
+}
diff --git a/Engine/Source/Runtime/IOS/MetalRHI/Private/MetalShaders.cpp b/Engine/Source/Runtime/IOS/MetalRHI/Private/MetalShaders.cpp
new file mode 100644
--- /dev/null
+++ b/Engine/Source/Runtime/IOS/MetalRHI/Private/MetalShaders.cpp
@@ -0,0 +1,575 @@
+// Copyright 1998-2014 Epic Games, Inc. All Rights Reserved.
+
+/*=============================================================================
+	MetalShaders.cpp: Metal shader RHI implementation.
+=============================================================================*/
+
+#include "MetalRHIPrivate.h"
+#include "MetalShaderResources.h"
+#include "MetalResources.h"
+
+/** Set to 1 to enable shader debugging (makes the driver save the shader source) */
+#define DEBUG_METAL_SHADERS (UE_BUILD_DEBUG || UE_BUILD_DEVELOPMENT)
+
+uint64 HitchyHashes[] = 
+{
+ 	0x09f873f002005ULL,
+ 	0x040050e020150ULL,
+ 	0x09f873e0a0150ULL,
+	0x11f873f002005ULL
+};
+
+
+class FMetalCompiledShaderKey
+{
+public:
+	FMetalCompiledShaderKey(
+		uint32 InCodeSize,
+		uint32 InCodeCRC
+		)
+		: CodeSize(InCodeSize)
+		, CodeCRC(InCodeCRC)
+	{}
+
+	friend bool operator ==(const FMetalCompiledShaderKey& A, const FMetalCompiledShaderKey& B)
+	{
+		return A.CodeSize == B.CodeSize && A.CodeCRC == B.CodeCRC;
+	}
+
+	friend uint32 GetTypeHash(const FMetalCompiledShaderKey &Key)
+	{
+		return GetTypeHash(Key.CodeSize) ^ GetTypeHash(Key.CodeCRC);
+	}
+
+private:
+	GLenum TypeEnum;
+	uint32 CodeSize;
+	uint32 CodeCRC;
+};
+
+typedef TMap<FMetalCompiledShaderKey, GLuint> FMetalCompiledShaderCache;
+
+static FMetalCompiledShaderCache& GetMetalCompiledShaderCache()
+{
+	static FMetalCompiledShaderCache CompiledShaderCache;
+	return CompiledShaderCache;
+}
+
+/** Initialization constructor. */
+template<typename BaseResourceType>
+TMetalBaseShader<BaseResourceType>::TMetalBaseShader(const TArray<uint8>& InCode)
+{
+	FMemoryReader Ar(InCode, true);
+
+	// was the shader already compiled offline?
+	uint8 OfflineCompiledFlag;
+	Ar << OfflineCompiledFlag;
+	check(OfflineCompiledFlag == 0 || OfflineCompiledFlag == 1);
+
+	// get the header
+	FMetalCodeHeader Header = { 0 };
+	Ar << Header;
+
+	// remember where the header ended and code (precompiled or source) begins
+	int32 CodeOffset = Ar.Tell();
+	const ANSICHAR* SourceCode = (ANSICHAR*)InCode.GetTypedData() + CodeOffset;
+
+	uint32 CodeLength, CodeCRC;
+	if (OfflineCompiledFlag)
+	{
+		CodeLength = InCode.Num() - CodeOffset;
+
+		// CRC the compiled code
+		CodeCRC = FCrc::MemCrc_DEPRECATED(InCode.GetTypedData() + CodeOffset, InCode.Num() - CodeOffset);
+	}
+	else
+	{
+		NSLog(@"Loaded an offline shader");
+
+		CodeLength = InCode.Num() - CodeOffset - 1;
+
+		// CRC the source
+		CodeCRC = FCrc::MemCrc_DEPRECATED(SourceCode, CodeLength);
+	}
+	FMetalCompiledShaderKey Key(CodeLength, CodeCRC);
+
+	// Find the existing compiled shader in the cache.
+	auto Resource = GetMetalCompiledShaderCache().FindRef(Key);
+//	if (!Resource)
+	{
+		id<MTLLibrary> Library;
+
+		if (OfflineCompiledFlag)
+		{
+			// allow GCD to copy the data into its own buffer
+			//		dispatch_data_t GCDBuffer = dispatch_data_create(InCode.GetTypedData() + CodeOffset, InCode.Num() - CodeOffset, nil, DISPATCH_DATA_DESTRUCTOR_DEFAULT);
+			uint32 BufferSize = InCode.Num() - CodeOffset;
+			void* Buffer = FMemory::Malloc( BufferSize );
+			FMemory::Memcpy( Buffer, InCode.GetTypedData() + CodeOffset, BufferSize );
+			dispatch_data_t GCDBuffer = dispatch_data_create(Buffer, BufferSize, dispatch_get_main_queue(), ^(void) { FMemory::Free(Buffer); } );
+
+			// load up the already compiled shader
+			NSError* Error;
+			Library = [FMetalManager::GetDevice() newLibraryWithData:GCDBuffer error:&Error];
+			if (Library == nil)
+			{
+				NSLog(@"Failed to create library: %@", Error);
+			}
+			
+			GlslCodeNSString = @"OFFLINE";
+			dispatch_release(GCDBuffer);
+		}
+		else
+		{
+			NSError* Error;
+			NSString* ShaderSource = [NSString stringWithUTF8String:SourceCode];
+
+	#if DEBUG_METAL_SHADERS
+			NSDictionary *PreprocessorMacros = @{ @"MTLSL_ENABLE_DEBUG_INFO" : @(1)};
+			MTLCompileOptions *CompileOptions = [[MTLCompileOptions alloc] init];
+			[CompileOptions setPreprocessorMacros : PreprocessorMacros];
+
+			Library = [FMetalManager::GetDevice() newLibraryWithSource:ShaderSource options : CompileOptions error : &Error];
+			[CompileOptions release];
+	#else
+			Library = [FMetalManager::GetDevice() newLibraryWithSource:ShaderSource options : nil error : &Error];
+	#endif
+
+			if (Library == nil)
+			{
+				NSLog(@"Failed to create library: %@", Error);
+				NSLog(@"*********** Error\n%@", ShaderSource);
+			}
+			else if (Error != nil)
+			{
+				// Warning...
+				NSLog(@"%@\n", Error);
+			}
+
+			GlslCode.Empty(CodeLength + 1);
+			GlslCode.AddUninitialized(CodeLength + 1);
+			FMemory::Memcpy(GlslCode.GetTypedData(), SourceCode, CodeLength + 1);
+			GlslCodeString = (ANSICHAR*)GlslCode.GetTypedData();
+			GlslCodeNSString = ShaderSource;
+
+#if !UE_BUILD_SHIPPING
+			[GlslCodeNSString retain];
+			TRACK_OBJECT(GlslCodeNSString);
+#endif
+		}
+
+		// assume there's only one function called 'Main', and use that to get the function from the library
+		Function = [Library newFunctionWithName:@"Main"];
+		[Library release];
+
+//		Resource = Resource;
+		Bindings = Header.Bindings;
+		BoundUniformBuffers.AddZeroed(Header.Bindings.NumUniformBuffers);
+		UniformBuffersCopyInfo = Header.UniformBuffersCopyInfo;
+	}
+}
+
+/** Destructor */
+template<typename BaseResourceType>
+TMetalBaseShader<BaseResourceType>::~TMetalBaseShader()
+{
+	UNTRACK_OBJECT(Function);
+	[Function release];
+	UNTRACK_OBJECT(GlslCodeNSString);
+	[GlslCodeNSString release];
+}
+
+
+FVertexShaderRHIRef FMetalDynamicRHI::RHICreateVertexShader(const TArray<uint8>& Code)
+{
+	FMetalVertexShader* Shader = new FMetalVertexShader(Code);
+	return Shader;
+}
+
+FPixelShaderRHIRef FMetalDynamicRHI::RHICreatePixelShader(const TArray<uint8>& Code)
+{
+	FMetalPixelShader* Shader = new FMetalPixelShader(Code);
+	return Shader;
+}
+
+FHullShaderRHIRef FMetalDynamicRHI::RHICreateHullShader(const TArray<uint8>& Code) 
+{ 
+	FMetalHullShader* Shader = new FMetalHullShader(Code);
+	return Shader;
+}
+
+FDomainShaderRHIRef FMetalDynamicRHI::RHICreateDomainShader(const TArray<uint8>& Code) 
+{ 
+	FMetalDomainShader* Shader = new FMetalDomainShader(Code);
+	return Shader;
+}
+
+FGeometryShaderRHIRef FMetalDynamicRHI::RHICreateGeometryShader(const TArray<uint8>& Code) 
+{ 
+	FMetalGeometryShader* Shader = new FMetalGeometryShader(Code);
+	return Shader;
+}
+
+FGeometryShaderRHIRef FMetalDynamicRHI::RHICreateGeometryShaderWithStreamOutput(const TArray<uint8>& Code, const FStreamOutElementList& ElementList,
+	uint32 NumStrides, const uint32* Strides, int32 RasterizedStream)
+{
+	checkf(0, TEXT("Not supported yet"));
+	return NULL;
+}
+
+FComputeShaderRHIRef FMetalDynamicRHI::RHICreateComputeShader(const TArray<uint8>& Code) 
+{ 
+	FMetalComputeShader* Shader = new FMetalComputeShader(Code);
+	return Shader;
+}
+
+FMetalBoundShaderState::FMetalBoundShaderState(
+			FVertexDeclarationRHIParamRef InVertexDeclarationRHI,
+			FVertexShaderRHIParamRef InVertexShaderRHI,
+			FPixelShaderRHIParamRef InPixelShaderRHI,
+			FHullShaderRHIParamRef InHullShaderRHI,
+			FDomainShaderRHIParamRef InDomainShaderRHI,
+			FGeometryShaderRHIParamRef InGeometryShaderRHI)
+	:	CacheLink(InVertexDeclarationRHI,InVertexShaderRHI,InPixelShaderRHI,InHullShaderRHI,InDomainShaderRHI,InGeometryShaderRHI,this)
+{
+	DYNAMIC_CAST_METGALRESOURCE(VertexDeclaration,InVertexDeclaration);
+	DYNAMIC_CAST_METGALRESOURCE(VertexShader,InVertexShader);
+	DYNAMIC_CAST_METGALRESOURCE(PixelShader,InPixelShader);
+	DYNAMIC_CAST_METGALRESOURCE(HullShader,InHullShader);
+	DYNAMIC_CAST_METGALRESOURCE(DomainShader,InDomainShader);
+	DYNAMIC_CAST_METGALRESOURCE(GeometryShader,InGeometryShader);
+
+	// cache everything
+	VertexDeclaration = InVertexDeclaration;
+	VertexShader = InVertexShader;
+	PixelShader = InPixelShader;
+	HullShader = InHullShader;
+	DomainShader = InDomainShader;
+	GeometryShader = InGeometryShader;
+
+#if !UE_BUILD_SHIPPING
+	extern FString GMaterialNameHack;
+	if (GFrameCounter > 3 || GMaterialNameHack == TEXT("M_Butterfly_OnGround") ||
+		GMaterialNameHack == TEXT("M_BillBoardBGTree") || GMaterialNameHack == TEXT("M_WaterSplash_Particulate"))
+	{
+		NSLog(@"===============================================================");
+		NSLog(@"Creating a BSS at runtime frame %lld... this may hitch! [%s, this = %p]", GFrameCounter, TCHAR_TO_ANSI(*GMaterialNameHack), this);
+		NSLog(@"Vertex declaration:");
+		FVertexDeclarationElementList& Elements = VertexDeclaration->Elements;
+		for (int32 i = 0; i < Elements.Num(); i++)
+		{
+			FVertexElement& Elem = Elements[i];
+			NSLog(@"   Elem %d: attr: %d, stream: %d, type: %d, stride: %d, offset: %d", i, Elem.AttributeIndex, Elem.StreamIndex, (uint32)Elem.Type, Elem.Stride, Elem.Offset);
+		}
+		
+		NSLog(@"\nVertexShader:");
+		NSLog(@"%@", VertexShader ? VertexShader->GlslCodeNSString : @"NONE");
+		NSLog(@"\nPixelShader:");
+		NSLog(@"%@", PixelShader ? PixelShader->GlslCodeNSString : @"NONE");
+		NSLog(@"===============================================================");
+	}
+#endif
+	
+	if (GFrameCounter < 2000)
+	{
+		for (int32 Index = 0; Index < ARRAY_COUNT(HitchyHashes); Index++)
+		{
+			FPipelineShadow Shadow;
+			Shadow.SetHash(HitchyHashes[Index]);
+			check(Shadow.GetHash() == HitchyHashes[Index]);
+	 
+			id<MTLRenderPipelineState> PipelineState = Shadow.CreatePipelineStateForBoundShaderState(this);
+			PipelineStates.Add(HitchyHashes[Index], PipelineState);
+		}
+	}
+	
+	AddRef();
+}
+
+FMetalBoundShaderState::~FMetalBoundShaderState()
+{
+	// free all of the pipeline state objects we have made
+	for (auto It = PipelineStates.CreateIterator(); It; ++It)
+	{
+		TRACK_OBJECT(It.Value());
+		[It.Value() release];
+	}
+}
+
+void FMetalBoundShaderState::PrepareToDraw(const FPipelineShadow& PipelineShadow)
+{
+#if NO_DRAW
+	return;
+#endif
+	
+	// generate a key for the current state
+	uint64 Hash = PipelineShadow.GetHash();
+	
+    
+	// have we made a matching state object yet?
+	id<MTLRenderPipelineState> PipelineState = PipelineStates.FindRef(Hash);
+	
+	// make one if not
+	if (PipelineState == nil)
+	{
+		PipelineState = PipelineShadow.CreatePipelineStateForBoundShaderState(this);
+		PipelineStates.Add(Hash, PipelineState);
+		
+#if !UE_BUILD_SHIPPING
+		if (GFrameCounter > 3)
+		{
+			NSLog(@"Created a hitchy pipeline state for hash %llx (this = %p)", Hash, this);
+		}
+#endif
+	}
+	
+	// set it now
+	[FMetalManager::GetContext() setRenderPipelineState:PipelineState];
+}
+
+FBoundShaderStateRHIRef FMetalDynamicRHI::RHICreateBoundShaderState(
+	FVertexDeclarationRHIParamRef VertexDeclarationRHI, 
+	FVertexShaderRHIParamRef VertexShaderRHI, 
+	FHullShaderRHIParamRef HullShaderRHI, 
+	FDomainShaderRHIParamRef DomainShaderRHI, 
+	FPixelShaderRHIParamRef PixelShaderRHI,
+	FGeometryShaderRHIParamRef GeometryShaderRHI
+	)
+{
+	// Check for an existing bound shader state which matches the parameters
+	FCachedBoundShaderStateLink* CachedBoundShaderStateLink = GetCachedBoundShaderState(
+		VertexDeclarationRHI,
+		VertexShaderRHI,
+		PixelShaderRHI,
+		HullShaderRHI,
+		DomainShaderRHI,
+		GeometryShaderRHI
+		);
+
+	if(CachedBoundShaderStateLink)
+	{
+		// If we've already created a bound shader state with these parameters, reuse it.
+		return CachedBoundShaderStateLink->BoundShaderState;
+	}
+	else
+	{
+		SCOPE_CYCLE_COUNTER(STAT_MetalBoundShaderStateTime);
+//		FBoundShaderStateKey Key(VertexDeclarationRHI,VertexShaderRHI,PixelShaderRHI,HullShaderRHI,DomainShaderRHI,GeometryShaderRHI);
+//		NSLog(@"BSS Key = %x", GetTypeHash(Key));
+		return new FMetalBoundShaderState(VertexDeclarationRHI,VertexShaderRHI,PixelShaderRHI,HullShaderRHI,DomainShaderRHI,GeometryShaderRHI);
+	}
+}
+
+
+
+
+
+
+
+
+FMetalShaderParameterCache::FMetalShaderParameterCache() :
+	GlobalUniformArraySize(-1)
+{
+	for (int32 ArrayIndex = 0; ArrayIndex < METAL_PACKED_TYPEINDEX_MAX; ++ArrayIndex)
+	{
+		PackedGlobalUniformDirty[ArrayIndex].LowVector = 0;
+		PackedGlobalUniformDirty[ArrayIndex].HighVector = 0;
+	}
+}
+
+void FMetalShaderParameterCache::InitializeResources(int32 UniformArraySize)
+{
+	check(GlobalUniformArraySize == -1);
+	
+	PackedGlobalUniforms[0] = (uint8*)FMemory::Malloc(UniformArraySize * METAL_PACKED_TYPEINDEX_MAX);
+	PackedUniformsScratch[0] = (uint8*)FMemory::Malloc(UniformArraySize * METAL_PACKED_TYPEINDEX_MAX);
+	
+	FMemory::Memzero(PackedGlobalUniforms[0], UniformArraySize * METAL_PACKED_TYPEINDEX_MAX);
+	FMemory::Memzero(PackedUniformsScratch[0], UniformArraySize * METAL_PACKED_TYPEINDEX_MAX);
+	for (int32 ArrayIndex = 1; ArrayIndex < METAL_PACKED_TYPEINDEX_MAX; ++ArrayIndex)
+	{
+		PackedGlobalUniforms[ArrayIndex] = PackedGlobalUniforms[ArrayIndex - 1] + UniformArraySize;
+		PackedUniformsScratch[ArrayIndex] = PackedUniformsScratch[ArrayIndex - 1] + UniformArraySize;
+	}
+	GlobalUniformArraySize = UniformArraySize;
+	
+	for (int32 ArrayIndex = 0; ArrayIndex < METAL_PACKED_TYPEINDEX_MAX; ++ArrayIndex)
+	{
+		PackedGlobalUniformDirty[ArrayIndex].LowVector = 0;
+		PackedGlobalUniformDirty[ArrayIndex].HighVector = 0;
+	}
+}
+
+/** Destructor. */
+FMetalShaderParameterCache::~FMetalShaderParameterCache()
+{
+	if (GlobalUniformArraySize > 0)
+	{
+		FMemory::Free(PackedUniformsScratch[0]);
+		FMemory::Free(PackedGlobalUniforms[0]);
+	}
+	
+	FMemory::MemZero(PackedUniformsScratch);
+	FMemory::MemZero(PackedGlobalUniforms);
+	
+	GlobalUniformArraySize = -1;
+}
+
+/**
+ * Marks all uniform arrays as dirty.
+ */
+void FMetalShaderParameterCache::MarkAllDirty()
+{
+	for (int32 ArrayIndex = 0; ArrayIndex < METAL_PACKED_TYPEINDEX_MAX; ++ArrayIndex)
+	{
+		PackedGlobalUniformDirty[ArrayIndex].LowVector = 0;
+		PackedGlobalUniformDirty[ArrayIndex].HighVector = 0;//UniformArraySize / 16;
+	}
+}
+
+const int SizeOfFloat4 = 4 * sizeof(float);
+
+/**
+ * Set parameter values.
+ */
+void FMetalShaderParameterCache::Set(uint32 BufferIndexName, uint32 ByteOffset, uint32 NumBytes, const void* NewValues)
+{
+	uint32 BufferIndex = MetalPackedTypeNameToTypeIndex(BufferIndexName);
+	check(GlobalUniformArraySize != -1);
+	check(BufferIndex < METAL_PACKED_TYPEINDEX_MAX);
+	check(ByteOffset + NumBytes <= (uint32)GlobalUniformArraySize);
+	PackedGlobalUniformDirty[BufferIndex].LowVector = FMath::Min(PackedGlobalUniformDirty[BufferIndex].LowVector, ByteOffset / SizeOfFloat4);
+	PackedGlobalUniformDirty[BufferIndex].HighVector = FMath::Max(PackedGlobalUniformDirty[BufferIndex].HighVector, (ByteOffset + NumBytes + SizeOfFloat4 - 1) / SizeOfFloat4);
+	FMemory::Memcpy(PackedGlobalUniforms[BufferIndex] + ByteOffset, NewValues, NumBytes);
+}
+
+void FMetalShaderParameterCache::CommitPackedGlobals(int32 Stage, const FMetalShaderBindings& Bindings)
+{
+	// copy the current uniform buffer into the ring buffer to submit
+	for (int32 Index = 0; Index < Bindings.PackedGlobalArrays.Num(); ++Index)
+	{
+        int32 UniformBufferIndex = Bindings.PackedGlobalArrays[Index].TypeIndex;
+ 
+		// is there any data that needs to be copied?
+		if (PackedGlobalUniformDirty[UniformBufferIndex].HighVector > 0)//PackedGlobalUniformDirty[UniformBufferIndex].LowVector)
+		{
+  			check(UniformBufferIndex == 0 || UniformBufferIndex == 1);
+            uint32 TotalSize = Bindings.PackedGlobalArrays[Index].Size;
+			uint32 SizeToUpload = PackedGlobalUniformDirty[UniformBufferIndex].HighVector * SizeOfFloat4;
+            check(SizeToUpload <= TotalSize);
+            SizeToUpload = TotalSize;
+ /*
+			if (UniformBufferIndex == METAL_PACKED_TYPEINDEX_MEDIUMP)
+			{
+				// Float4 -> Half4
+				SizeToUpload /= 2;
+			}
+*/
+			// allocate memory in the ring buffer
+			id<MTLBuffer> RingBuffer = FMetalManager::Get()->GetRingBuffer();
+			uint32 Offset = FMetalManager::Get()->AllocateFromRingBuffer(TotalSize);
+/*
+			// copy into the middle of the ring buffer
+			if (UniformBufferIndex == METAL_PACKED_TYPEINDEX_MEDIUMP)
+			{
+				// Convert to half
+				const int NumVectors = PackedGlobalUniformDirty[UniformBufferIndex].HighVector;
+				const auto* RESTRICT SourceF4 = (FVector4* )PackedGlobalUniforms[UniformBufferIndex];
+				auto* RESTRICT Dest = (FFloat16*)((uint8*)[RingBuffer contents]) + Offset;
+				for (int Vector = 0; Vector < NumVectors; ++Vector, ++SourceF4, Dest += 4)
+				{
+					Dest[0].Set(SourceF4->X);
+					Dest[1].Set(SourceF4->Y);
+					Dest[2].Set(SourceF4->Z);
+					Dest[3].Set(SourceF4->W);
+				}
+			}
+			else
+ */
+			{
+				//check(UniformBufferIndex != METAL_PACKED_TYPEINDEX_LOWP);
+				FMemory::Memcpy(((uint8*)[RingBuffer contents]) + Offset, PackedGlobalUniforms[UniformBufferIndex], SizeToUpload);
+			}
+			
+			switch (Stage)
+			{
+				case METAL_SHADER_STAGE_VERTEX:
+//					NSLog(@"Uploading %d bytes to vertex", SizeToUpload);
+					[FMetalManager::GetContext() setVertexBuffer:RingBuffer offset:Offset atIndex:UniformBufferIndex];
+					break;
+
+				case METAL_SHADER_STAGE_PIXEL:
+//					NSLog(@"Uploading %d bytes to pixel", SizeToUpload);
+					[FMetalManager::GetContext() setFragmentBuffer:RingBuffer offset : Offset atIndex : UniformBufferIndex];
+					break;
+					
+				
+				default:check(0);
+			}
+
+			// mark as clean
+			PackedGlobalUniformDirty[UniformBufferIndex].HighVector = 0;
+		}
+	}
+}
+
+void FMetalShaderParameterCache::CommitPackedUniformBuffers(TRefCountPtr<FMetalBoundShaderState> BoundShaderState, int32 Stage, const TArray< TRefCountPtr<FRHIUniformBuffer> >& RHIUniformBuffers, const TArray<FMetalUniformBufferCopyInfo>& UniformBuffersCopyInfo)
+{
+//	SCOPE_CYCLE_COUNTER(STAT_MetalConstantBufferUpdateTime);
+	// Uniform Buffers are split into precision/type; the list of RHI UBs is traversed and if a new one was set, its
+	// contents are copied per precision/type into corresponding scratch buffers which are then uploaded to the program
+	if (Stage == METAL_SHADER_STAGE_PIXEL && !IsValidRef(BoundShaderState->PixelShader))
+	{
+		return;
+	}
+
+	auto& Bindings = Stage == METAL_SHADER_STAGE_VERTEX ? BoundShaderState->VertexShader->Bindings : BoundShaderState->PixelShader->Bindings;
+	check(Bindings.NumUniformBuffers == RHIUniformBuffers.Num());
+	if (!Bindings.bHasRegularUniformBuffers)
+	{
+		int32 LastInfoIndex = 0;
+		for (int32 BufferIndex = 0; BufferIndex < RHIUniformBuffers.Num(); ++BufferIndex)
+		{
+			const FRHIUniformBuffer* RHIUniformBuffer = RHIUniformBuffers[BufferIndex];
+			check(RHIUniformBuffer);
+			FMetalUniformBuffer* EmulatedUniformBuffer = (FMetalUniformBuffer*)RHIUniformBuffer;
+			const uint32* RESTRICT SourceData = (uint32*)((uint8*)[EmulatedUniformBuffer->Buffer contents] + EmulatedUniformBuffer->Offset);//->Data.GetTypedData();
+			for (int32 InfoIndex = LastInfoIndex; InfoIndex < UniformBuffersCopyInfo.Num(); ++InfoIndex)
+			{
+				const FMetalUniformBufferCopyInfo& Info = UniformBuffersCopyInfo[InfoIndex];
+				if (Info.SourceUBIndex == BufferIndex)
+				{
+					float* RESTRICT ScratchMem = (float*)PackedGlobalUniforms[Info.DestUBTypeIndex];
+					ScratchMem += Info.DestOffsetInFloats;
+					FMemory::Memcpy(ScratchMem, SourceData + Info.SourceOffsetInFloats, Info.SizeInFloats * sizeof(float));
+					PackedGlobalUniformDirty[Info.DestUBTypeIndex].LowVector = FMath::Min(PackedGlobalUniformDirty[Info.DestUBTypeIndex].LowVector, uint32(Info.DestOffsetInFloats / SizeOfFloat4));
+					PackedGlobalUniformDirty[Info.DestUBTypeIndex].HighVector = FMath::Max(PackedGlobalUniformDirty[Info.DestUBTypeIndex].HighVector, uint32(((Info.DestOffsetInFloats + Info.SizeInFloats) * sizeof(float) + SizeOfFloat4 - 1) / SizeOfFloat4));
+				}
+				else
+				{
+					LastInfoIndex = InfoIndex;
+					break;
+				}
+			}
+		}
+	}
+	else
+	{
+		for (int32 BufferIndex = 0; BufferIndex < RHIUniformBuffers.Num(); ++BufferIndex)
+		{
+			const FRHIUniformBuffer* RHIUniformBuffer = RHIUniformBuffers[BufferIndex];
+			if (RHIUniformBuffer)
+			{
+				auto* UB = (FMetalUniformBuffer*)RHIUniformBuffer;
+				if (Stage == METAL_SHADER_STAGE_VERTEX)
+				{
+					[FMetalManager::GetContext() setVertexBuffer:UB->Buffer offset:UB->Offset atIndex:BufferIndex];
+				}
+				else
+				{
+					[FMetalManager::GetContext() setFragmentBuffer:UB->Buffer offset:UB->Offset atIndex:BufferIndex];
+				}
+			}
+		}
+	}
+}
diff --git a/Engine/Source/Runtime/IOS/MetalRHI/Private/MetalState.cpp b/Engine/Source/Runtime/IOS/MetalRHI/Private/MetalState.cpp
new file mode 100644
--- /dev/null
+++ b/Engine/Source/Runtime/IOS/MetalRHI/Private/MetalState.cpp
@@ -0,0 +1,318 @@
+// Copyright 1998-2014 Epic Games, Inc. All Rights Reserved.
+
+/*=============================================================================
+	MetalState.cpp: Metal state implementation.
+=============================================================================*/
+
+#include "MetalRHIPrivate.h"
+
+/**
+ * Translate from UE4 enums to Metal enums (leaving these in Metal just for help in setting them up)
+ */
+static MTLSamplerMipFilter TranslateMipFilterMode(ESamplerFilter Filter)
+{
+	switch (Filter)
+	{
+		case SF_Point:	return MTLSamplerMipFilterNearest;
+		default:		return MTLSamplerMipFilterLinear;
+	}
+}
+
+static MTLSamplerMinMagFilter TranslateFilterMode(ESamplerFilter Filter)
+{
+	switch (Filter)
+	{
+		case SF_Point:				return MTLSamplerMinMagFilterNearest;
+		case SF_AnisotropicPoint:	return MTLSamplerMinMagFilterLinear;
+		case SF_AnisotropicLinear:	return MTLSamplerMinMagFilterLinear;
+		default:					return MTLSamplerMinMagFilterLinear;
+	}
+}
+
+static uint32 GetMaxAnisotropy(ESamplerFilter Filter, uint32 MaxAniso)
+{
+	switch (Filter)
+	{
+		case SF_AnisotropicPoint:
+		case SF_AnisotropicLinear:	return MaxAniso > 0 ? MaxAniso : GetCachedScalabilityCVars().MaxAnisotropy;
+		default:					return 1;
+	}
+}
+
+static MTLSamplerMinMagFilter TranslateZFilterMode(ESamplerFilter Filter)
+{	switch (Filter)
+	{
+		case SF_Point:				return MTLSamplerMinMagFilterNearest;
+		case SF_AnisotropicPoint:	return MTLSamplerMinMagFilterNearest;
+		case SF_AnisotropicLinear:	return MTLSamplerMinMagFilterLinear;
+		default:					return MTLSamplerMinMagFilterLinear;
+	}
+}
+
+static MTLSamplerAddressMode TranslateWrapMode(ESamplerAddressMode AddressMode)
+{
+	switch (AddressMode)
+	{
+		case AM_Clamp:	return MTLSamplerAddressModeClampToEdge;
+		case AM_Mirror: return MTLSamplerAddressModeMirrorRepeat;
+		case AM_Border: return MTLSamplerAddressModeClampToEdge;
+		default:		return MTLSamplerAddressModeRepeat;
+	}
+}
+
+static MTLCompareFunction TranslateCompareFunction(ECompareFunction CompareFunction)
+{
+	switch(CompareFunction)
+	{
+		case CF_Less:			return MTLCompareFunctionLess;
+		case CF_LessEqual:		return MTLCompareFunctionLessEqual;
+		case CF_Greater:		return MTLCompareFunctionGreater;
+		case CF_GreaterEqual:	return MTLCompareFunctionGreaterEqual;
+		case CF_Equal:			return MTLCompareFunctionEqual;
+		case CF_NotEqual:		return MTLCompareFunctionNotEqual;
+		case CF_Never:			return MTLCompareFunctionNever;
+		default:				return MTLCompareFunctionAlways;
+	};
+}
+
+static int32 TranslateSamplerCompareFunction(ESamplerCompareFunction SamplerComparisonFunction)
+{
+	switch(SamplerComparisonFunction)
+	{
+		case SCF_Less:		return 0;
+		case SCF_Never:		return 0;
+		default:			return 0;
+	};
+}
+
+static MTLStencilOperation TranslateStencilOp(EStencilOp StencilOp)
+{
+	switch(StencilOp)
+	{
+		case SO_Zero:				return MTLStencilOperationZero;
+		case SO_Replace:			return MTLStencilOperationReplace;
+		case SO_SaturatedIncrement:	return MTLStencilOperationIncrementClamp;
+		case SO_SaturatedDecrement:	return MTLStencilOperationDecrementClamp;
+		case SO_Invert:				return MTLStencilOperationInvert;
+		case SO_Increment:			return MTLStencilOperationIncrementWrap;
+		case SO_Decrement:			return MTLStencilOperationDecrementWrap;
+		default:					return MTLStencilOperationKeep;
+	};
+}
+
+static MTLBlendOperation TranslateBlendOp(EBlendOperation BlendOp)
+{
+	switch(BlendOp)
+	{
+		case BO_Subtract:	return MTLBlendOperationSubtract;
+		case BO_Min:		return MTLBlendOperationMin;
+		case BO_Max:		return MTLBlendOperationMax;
+		default:			return MTLBlendOperationAdd;
+	};
+}
+
+
+static MTLBlendFactor TranslateBlendFactor(EBlendFactor BlendFactor)
+{
+	switch(BlendFactor)
+	{
+		case BF_One:					return MTLBlendFactorOne;
+		case BF_SourceColor:			return MTLBlendFactorSourceColor;
+		case BF_InverseSourceColor:		return MTLBlendFactorOneMinusSourceColor;
+		case BF_SourceAlpha:			return MTLBlendFactorSourceAlpha;
+		case BF_InverseSourceAlpha:		return MTLBlendFactorOneMinusSourceAlpha;
+		case BF_DestAlpha:				return MTLBlendFactorDestinationAlpha;
+		case BF_InverseDestAlpha:		return MTLBlendFactorOneMinusDestinationAlpha;
+		case BF_DestColor:				return MTLBlendFactorDestinationColor;
+		case BF_InverseDestColor:		return MTLBlendFactorOneMinusDestinationColor;
+		default:						return MTLBlendFactorZero;
+	};
+}
+
+static MTLColorWriteMask TranslateWriteMask(EColorWriteMask WriteMask)
+{
+	uint32 Result = 0;
+	Result |= (WriteMask & CW_RED) ? (MTLColorWriteMaskRed) : 0;
+	Result |= (WriteMask & CW_GREEN) ? (MTLColorWriteMaskGreen) : 0;
+	Result |= (WriteMask & CW_BLUE) ? (MTLColorWriteMaskBlue) : 0;
+	Result |= (WriteMask & CW_ALPHA) ? (MTLColorWriteMaskAlpha) : 0;
+	
+	return (MTLColorWriteMask)Result;
+}
+
+
+
+
+
+FMetalSamplerState::FMetalSamplerState(const FSamplerStateInitializerRHI& Initializer)
+{
+	MTLSamplerDescriptor* Desc = [[MTLSamplerDescriptor alloc] init];
+	
+	Desc.minFilter = Desc.magFilter = TranslateFilterMode(Initializer.Filter);
+	Desc.mipFilter = TranslateMipFilterMode(Initializer.Filter);
+	Desc.maxAnisotropy = GetMaxAnisotropy(Initializer.Filter, Initializer.MaxAnisotropy);
+	Desc.sAddressMode = TranslateWrapMode(Initializer.AddressU);
+	Desc.tAddressMode = TranslateWrapMode(Initializer.AddressV);
+	Desc.rAddressMode = TranslateWrapMode(Initializer.AddressW);
+	Desc.lodMinClamp = Initializer.MinMipLevel;
+	Desc.lodMaxClamp = Initializer.MaxMipLevel;
+	
+	State = [FMetalManager::GetDevice() newSamplerStateWithDescriptor:Desc];
+	[Desc release];
+	TRACK_OBJECT(State);
+}
+
+FMetalSamplerState::~FMetalSamplerState()
+{
+	UNTRACK_OBJECT(State);
+	[State release];
+}
+
+FMetalRasterizerState::FMetalRasterizerState(const FRasterizerStateInitializerRHI& Initializer)
+{
+	State = Initializer;
+}
+
+FMetalRasterizerState::~FMetalRasterizerState()
+{
+	
+}
+
+void FMetalRasterizerState::Set()
+{
+	FMetalManager::Get()->SetRasterizerState(State);
+}
+
+FMetalDepthStencilState::FMetalDepthStencilState(const FDepthStencilStateInitializerRHI& Initializer)
+{
+	MTLDepthStencilDescriptor* Desc = [[MTLDepthStencilDescriptor alloc] init];
+	Desc.depthCompareFunction = TranslateCompareFunction(Initializer.DepthTest);
+	Desc.depthWriteEnabled = Initializer.bEnableDepthWrite;
+	
+	if (Initializer.bEnableFrontFaceStencil)
+	{
+		// set up front face stencil operations
+		MTLStencilDescriptor* Stencil = [[MTLStencilDescriptor alloc] init];
+		Stencil.stencilCompareFunction = TranslateCompareFunction(Initializer.FrontFaceStencilTest);
+		Stencil.stencilFailureOperation = TranslateStencilOp(Initializer.FrontFaceStencilFailStencilOp);
+		Stencil.depthFailureOperation = TranslateStencilOp(Initializer.FrontFaceStencilFailStencilOp);
+		Stencil.depthStencilPassOperation = TranslateStencilOp(Initializer.FrontFacePassStencilOp);
+		Stencil.readMask = Initializer.StencilReadMask;
+		Stencil.writeMask = Initializer.StencilWriteMask;
+		Desc.frontFaceStencilDescriptor = Stencil;
+		
+		[Stencil release];
+	}
+	
+	if (Initializer.bEnableFrontFaceStencil)
+	{
+		// set up back face stencil operations
+		MTLStencilDescriptor* Stencil = [[MTLStencilDescriptor alloc] init];
+		Stencil.stencilCompareFunction = TranslateCompareFunction(Initializer.BackFaceStencilTest);
+		Stencil.stencilFailureOperation = TranslateStencilOp(Initializer.BackFaceStencilFailStencilOp);
+		Stencil.depthFailureOperation = TranslateStencilOp(Initializer.BackFaceStencilFailStencilOp);
+		Stencil.depthStencilPassOperation = TranslateStencilOp(Initializer.BackFacePassStencilOp);
+		Stencil.readMask = Initializer.StencilReadMask;
+		Stencil.writeMask = Initializer.StencilWriteMask;
+		Desc.backFaceStencilDescriptor = Stencil;
+		
+		[Stencil release];
+	}
+			
+	// bake out the descriptor
+	State = [FMetalManager::GetDevice() newDepthStencilStateWithDescriptor:Desc];
+	TRACK_OBJECT(State);
+	[Desc release];
+	
+	// cache some pipeline state info
+	bIsDepthWriteEnabled = Initializer.bEnableDepthWrite;
+	bIsStencilWriteEnabled = Initializer.bEnableFrontFaceStencil || Initializer.bEnableBackFaceStencil;
+}
+
+FMetalDepthStencilState::~FMetalDepthStencilState()
+{
+	UNTRACK_OBJECT(State);
+	[State release];
+}
+
+
+
+void FMetalDepthStencilState::Set()
+{
+	//activate the state
+	[FMetalManager::GetContext() setDepthStencilState:State];
+	
+	// update the pipeline state object
+	FMetalManager::Get()->SetDepthStencilWriteEnabled(bIsDepthWriteEnabled, bIsStencilWriteEnabled);
+}
+
+FMetalBlendState::FMetalBlendState(const FBlendStateInitializerRHI& Initializer)
+{
+	for(uint32 RenderTargetIndex = 0;RenderTargetIndex < MaxMetalRenderTargets; ++RenderTargetIndex)
+	{
+		// which initializer to use
+		const FBlendStateInitializerRHI::FRenderTarget& Init =
+			Initializer.bUseIndependentRenderTargetBlendStates
+				? Initializer.RenderTargets[RenderTargetIndex]
+				: Initializer.RenderTargets[0];
+
+		// make a new blend state
+		MTLBlendDescriptor* BlendState = [[MTLBlendDescriptor alloc] init];
+		TRACK_OBJECT(BlendState);
+
+		// set values
+		BlendState.blendingEnabled =
+			Init.ColorBlendOp != BO_Add || Init.ColorDestBlend != BF_Zero || Init.ColorSrcBlend != BF_One ||
+			Init.AlphaBlendOp != BO_Add || Init.AlphaDestBlend != BF_Zero || Init.AlphaSrcBlend != BF_One;
+		BlendState.sourceRGBBlendFactor = TranslateBlendFactor(Init.ColorSrcBlend);
+		BlendState.destinationRGBBlendFactor = TranslateBlendFactor(Init.ColorDestBlend);
+		BlendState.rgbBlendOperation = TranslateBlendOp(Init.ColorBlendOp);
+		BlendState.sourceAlphaBlendFactor = TranslateBlendFactor(Init.AlphaSrcBlend);
+		BlendState.destinationAlphaBlendFactor = TranslateBlendFactor(Init.AlphaDestBlend);
+		BlendState.alphaBlendOperation = TranslateBlendOp(Init.AlphaBlendOp);
+		BlendState.writeMask = TranslateWriteMask(Init.ColorWriteMask);
+		
+		RenderTargetStates[RenderTargetIndex] = BlendState;
+	}
+}
+
+FMetalBlendState::~FMetalBlendState()
+{
+	for(uint32 RenderTargetIndex = 0;RenderTargetIndex < MaxMetalRenderTargets; ++RenderTargetIndex)
+	{
+		UNTRACK_OBJECT(RenderTargetStates[RenderTargetIndex]);
+		[RenderTargetStates[RenderTargetIndex] release];
+	}
+}
+
+void FMetalBlendState::Set()
+{
+	// update the pipeline state, there's nothing else to do for blend state
+	FMetalManager::Get()->SetBlendState(this);
+}
+
+
+
+
+
+FSamplerStateRHIRef FMetalDynamicRHI::RHICreateSamplerState(const FSamplerStateInitializerRHI& Initializer)
+{
+	return new FMetalSamplerState(Initializer);
+}
+
+FRasterizerStateRHIRef FMetalDynamicRHI::RHICreateRasterizerState(const FRasterizerStateInitializerRHI& Initializer)
+{
+    return new FMetalRasterizerState(Initializer);
+}
+
+FDepthStencilStateRHIRef FMetalDynamicRHI::RHICreateDepthStencilState(const FDepthStencilStateInitializerRHI& Initializer)
+{
+	return new FMetalDepthStencilState(Initializer);
+}
+
+
+FBlendStateRHIRef FMetalDynamicRHI::RHICreateBlendState(const FBlendStateInitializerRHI& Initializer)
+{
+	return new FMetalBlendState(Initializer);
+}
+
diff --git a/Engine/Source/Runtime/IOS/MetalRHI/Private/MetalStructuredBuffer.cpp b/Engine/Source/Runtime/IOS/MetalRHI/Private/MetalStructuredBuffer.cpp
new file mode 100644
--- /dev/null
+++ b/Engine/Source/Runtime/IOS/MetalRHI/Private/MetalStructuredBuffer.cpp
@@ -0,0 +1,43 @@
+// Copyright 1998-2014 Epic Games, Inc. All Rights Reserved.
+
+
+#include "MetalRHIPrivate.h"
+
+
+FMetalStructuredBuffer::FMetalStructuredBuffer(uint32 Stride, uint32 Size, FResourceArrayInterface* ResourceArray, uint32 Usage)
+	: FRHIStructuredBuffer(Stride, Size, Usage)
+{
+	check((Size % Stride) == 0);
+
+	// copy any resources to the CPU address
+	if (ResourceArray)
+	{
+// 		FMemory::Memcpy( , ResourceArray->GetResourceData(), Size);
+		ResourceArray->Discard();
+	}
+}
+
+FMetalStructuredBuffer::~FMetalStructuredBuffer()
+{
+
+}
+
+
+
+
+FStructuredBufferRHIRef FMetalDynamicRHI::RHICreateStructuredBuffer(uint32 Stride,uint32 Size,uint32 InUsage,FRHIResourceCreateInfo& CreateInfo)
+{
+	return new FMetalStructuredBuffer(Stride, Size, CreateInfo.ResourceArray, InUsage);
+}
+
+void* FMetalDynamicRHI::RHILockStructuredBuffer(FStructuredBufferRHIParamRef StructuredBufferRHI,uint32 Offset,uint32 Size,EResourceLockMode LockMode)
+{
+	DYNAMIC_CAST_METGALRESOURCE(StructuredBuffer,StructuredBuffer);
+
+	return NULL;
+}
+
+void FMetalDynamicRHI::RHIUnlockStructuredBuffer(FStructuredBufferRHIParamRef StructuredBufferRHI)
+{
+
+}
diff --git a/Engine/Source/Runtime/IOS/MetalRHI/Private/MetalTexture.cpp b/Engine/Source/Runtime/IOS/MetalRHI/Private/MetalTexture.cpp
new file mode 100644
--- /dev/null
+++ b/Engine/Source/Runtime/IOS/MetalRHI/Private/MetalTexture.cpp
@@ -0,0 +1,423 @@
+// Copyright 1998-2014 Epic Games, Inc. All Rights Reserved.
+
+/*=============================================================================
+	MetalVertexBuffer.cpp: Metal texture RHI implementation.
+=============================================================================*/
+
+#include "MetalRHIPrivate.h"
+
+FORCEINLINE int32 GetMetalCubeFace(ECubeFace Face)
+{
+	switch (Face)
+	{
+		case CubeFace_NegX:	return 0;
+		case CubeFace_NegY:	return 1;
+		case CubeFace_NegZ:	return 2;
+		case CubeFace_PosX:	return 3;
+		case CubeFace_PosY:	return 4;
+		case CubeFace_PosZ:	return 5;
+		default:			return -1;
+	}
+}
+
+/** Texture reference class. */
+class FMetalTextureReference : public FRHITextureReference
+{
+public:
+	explicit FMetalTextureReference()
+		: FRHITextureReference()
+	{}
+
+	// IRefCountedObject interface.
+	virtual uint32 AddRef() const
+	{
+		return FRefCountedObject::AddRef();
+	}
+	virtual uint32 Release() const
+	{
+		return FRefCountedObject::Release();
+	}
+	virtual uint32 GetRefCount() const
+	{
+		return FRefCountedObject::GetRefCount();
+	}
+
+	void SetReferencedTexture(FRHITexture* InTexture)
+	{
+		FRHITextureReference::SetReferencedTexture(InTexture);
+	}
+};
+
+/** Given a pointer to a RHI texture that was created by the Metal RHI, returns a pointer to the FMetalTextureBase it encapsulates. */
+FMetalSurface& GetMetalSurfaceFromRHITexture(FRHITexture* Texture)
+{
+	if(Texture->GetTexture2D())
+	{
+		return ((FMetalTexture2D*)Texture)->Surface;
+	}
+	else if(Texture->GetTexture2DArray())
+	{
+		return ((FMetalTexture2DArray*)Texture)->Surface;
+	}
+	else if(Texture->GetTexture3D())
+	{
+		return ((FMetalTexture3D*)Texture)->Surface;
+	}
+	else if(Texture->GetTextureCube())
+	{
+		return ((FMetalTextureCube*)Texture)->Surface;
+	}
+	else if(Texture->GetTextureReference())
+	{
+		return GetMetalSurfaceFromRHITexture(static_cast<FMetalTextureReference*>(Texture)->GetReferencedTexture());
+	}
+	else
+	{
+		UE_LOG(LogMetal, Fatal, TEXT("Unknown RHI texture type"));
+		return ((FMetalTexture2D*)Texture)->Surface;
+	}
+}
+
+FMetalSurface::FMetalSurface(ERHIResourceType ResourceType, EPixelFormat Format, uint32 InSizeX, uint32 InSizeY, uint32 InSizeZ, uint32 NumSamples, bool bArray, uint32 ArraySize, uint32 NumMips, uint32 InFlags, FResourceBulkDataInterface* BulkData)
+	: PixelFormat(Format)
+    , MSAATexture(nil)
+	, SizeX(InSizeX)
+	, SizeY(InSizeY)
+	, SizeZ(InSizeZ)
+	, bIsCubemap(false)
+	, Flags(InFlags)
+{
+	// the special back buffer surface will be updated in FMetalManager::BeginFrame - no need to set the texture here
+	if (Flags & TexCreate_Presentable)
+	{
+		return;
+	}
+	
+	MTLPixelFormat MTLFormat = (MTLPixelFormat)GPixelFormats[Format].PlatformFormat;
+	MTLTextureDescriptor* Desc;
+	
+	if (ResourceType == RRT_TextureCube)
+	{
+		Desc = [MTLTextureDescriptor textureCubeDescriptorWithPixelFormat:MTLFormat
+																	size:SizeX
+															   mipmapped:(NumMips > 1)];
+		bIsCubemap = true;
+	}
+	else
+	{
+		Desc = [MTLTextureDescriptor texture2DDescriptorWithPixelFormat:MTLFormat
+																 width:SizeX
+																height:SizeY
+															 mipmapped:(NumMips > 1)];
+		Desc.depth = SizeZ;
+	}
+
+	// flesh out the descriptor
+	if (bArray)
+	{
+		Desc.arrayLength = ArraySize;
+	}
+	Desc.mipmapLevelCount = NumMips;
+	
+	
+	Texture = [FMetalManager::GetDevice() newTextureWithDescriptor:Desc];
+	if (Texture == nil)
+	{
+		NSLog(@"Failed to create texture, desc  %@", Desc);
+	}
+	TRACK_OBJECT(Texture);
+	
+	// upload existing bulkdata
+	if (BulkData)
+	{
+		UE_LOG(LogIOS, Display, TEXT("Got a bulk data texture, with %d mips"), NumMips);
+		// copy over bulk data
+		// Lock
+		// Memcpy
+		// Unlock
+
+		// bulk data can be unloaded now
+		BulkData->Discard();
+	}
+	
+	LockedMemory = NULL;
+
+	if (!FParse::Param(FCommandLine::Get(), TEXT("nomsaa")))
+	{
+		if (NumSamples > 1)
+		{
+			check(Flags & (TexCreate_RenderTargetable | TexCreate_DepthStencilTargetable));
+			Desc.textureType = MTLTextureType2DMultisample;
+	
+			// allow commandline to override
+			FParse::Value(FCommandLine::Get(), TEXT("msaa="), NumSamples);
+			Desc.sampleCount = NumSamples;
+
+			MSAATexture = [FMetalManager::GetDevice() newTextureWithDescriptor:Desc];
+			if (Format == PF_DepthStencil)
+			{
+				[Texture release];
+				Texture = MSAATexture;
+			}
+        
+			NSLog(@"Creating %dx MSAA %d x %d %s surface", (int32)Desc.sampleCount, SizeX, SizeY, (Flags & TexCreate_RenderTargetable) ? "Color" : "Depth");
+			if (MSAATexture == nil)
+			{
+				NSLog(@"Failed to create texture, desc  %@", Desc);
+			}
+			TRACK_OBJECT(MSAATexture);
+		}
+	}
+}
+
+
+FMetalSurface::~FMetalSurface()
+{
+    if (MSAATexture != nil)
+    {
+		FMetalManager::ReleaseObject(MSAATexture);
+    }
+    
+	if (!(Flags & TexCreate_Presentable))
+	{
+		FMetalManager::ReleaseObject(Texture);
+	}
+}
+
+
+void* FMetalSurface::Lock(uint32 MipIndex, uint32 ArrayIndex, EResourceLockMode LockMode, uint32& DestStride)
+{
+	// Calculate the dimensions of the mip-map.
+	const uint32 BlockSizeX = GPixelFormats[PixelFormat].BlockSizeX;
+	const uint32 BlockSizeY = GPixelFormats[PixelFormat].BlockSizeY;
+	const uint32 BlockBytes = GPixelFormats[PixelFormat].BlockBytes;
+	const uint32 MipSizeX = FMath::Max(SizeX >> MipIndex,BlockSizeX);
+	const uint32 MipSizeY = FMath::Max(SizeY >> MipIndex,BlockSizeY);
+	uint32 NumBlocksX = (MipSizeX + BlockSizeX - 1) / BlockSizeX;
+	uint32 NumBlocksY = (MipSizeY + BlockSizeY - 1) / BlockSizeY;
+	if ( PixelFormat == PF_PVRTC2 || PixelFormat == PF_PVRTC4 )
+	{
+		// PVRTC has minimum 2 blocks width and height
+		NumBlocksX = FMath::Max<uint32>(NumBlocksX, 2);
+		NumBlocksY = FMath::Max<uint32>(NumBlocksY, 2);
+	}
+	const uint32 MipBytes = NumBlocksX * NumBlocksY * BlockBytes;
+	DestStride = NumBlocksX * BlockBytes;
+	
+	// allocate some temporary memory
+	check(LockedMemory == NULL);
+	LockedMemory = FMemory::Malloc(MipBytes);
+	if (LockMode != RLM_WriteOnly)
+	{
+		// [Texture readPixels ...];
+	}
+
+	return LockedMemory;
+}
+
+void FMetalSurface::Unlock(uint32 MipIndex, uint32 ArrayIndex)
+{
+	const uint32 BlockSizeX = GPixelFormats[PixelFormat].BlockSizeX;
+	const uint32 BlockSizeY = GPixelFormats[PixelFormat].BlockSizeY;
+	const uint32 BlockBytes = GPixelFormats[PixelFormat].BlockBytes;
+	const uint32 MipSizeX = FMath::Max(SizeX >> MipIndex,BlockSizeX);
+	const uint32 MipSizeY = FMath::Max(SizeY >> MipIndex,BlockSizeY);
+	uint32 NumBlocksX = (MipSizeX + BlockSizeX - 1) / BlockSizeX;
+	uint32 NumBlocksY = (MipSizeY + BlockSizeY - 1) / BlockSizeY;
+	uint32 MipBytes = NumBlocksX * NumBlocksY * BlockBytes;
+	uint32 Stride = NumBlocksX * BlockBytes;
+	if (PixelFormat == PF_PVRTC2 || PixelFormat == PF_PVRTC4)
+	{
+		MipBytes = 0;
+		Stride = 0;
+	}
+
+	// upload the texture to the texture slice
+	// @todo urban: The MipSizeZ would be needed for a volume texture - do we care?
+	[Texture copyFromPixels:LockedMemory rowBytes:Stride imageBytes:MipBytes toSlice:ArrayIndex mipmapLevel:MipIndex origin:MTLOriginMake(0, 0, 0) size:MTLSizeMake(FMath::Max<uint32>(SizeX>>MipIndex, 1), FMath::Max<uint32>(SizeY>>MipIndex, 1), 1)];
+	
+	FMemory::Free(LockedMemory);
+	LockedMemory = NULL;
+}
+
+uint32 FMetalSurface::GetMemorySize()
+{
+	return 0;
+}
+
+/*-----------------------------------------------------------------------------
+	Texture allocator support.
+-----------------------------------------------------------------------------*/
+
+void FMetalDynamicRHI::RHIGetTextureMemoryStats(FTextureMemoryStats& OutStats)
+{
+
+}
+
+bool FMetalDynamicRHI::RHIGetTextureMemoryVisualizeData( FColor* /*TextureData*/, int32 /*SizeX*/, int32 /*SizeY*/, int32 /*Pitch*/, int32 /*PixelSize*/ )
+{
+	return false;
+}
+
+uint32 FMetalDynamicRHI::RHIComputeMemorySize(FTextureRHIParamRef TextureRHI)
+{
+	if(!TextureRHI)
+	{
+		return 0;
+	}
+
+	return GetMetalSurfaceFromRHITexture(TextureRHI).GetMemorySize();
+}
+
+/*-----------------------------------------------------------------------------
+	2D texture support.
+-----------------------------------------------------------------------------*/
+
+FTexture2DRHIRef FMetalDynamicRHI::RHICreateTexture2D(uint32 SizeX, uint32 SizeY, uint8 Format, uint32 NumMips, uint32 NumSamples, uint32 Flags, FRHIResourceCreateInfo& CreateInfo)
+{
+	return new FMetalTexture2D((EPixelFormat)Format, SizeX, SizeY, NumMips, NumSamples, Flags, CreateInfo.BulkData);
+}
+
+FTexture2DRHIRef FMetalDynamicRHI::RHIAsyncCreateTexture2D(uint32 SizeX,uint32 SizeY,uint8 Format,uint32 NumMips,uint32 Flags,void** InitialMipData,uint32 NumInitialMips)
+{
+	UE_LOG(LogMetal, Fatal, TEXT("RHIAsyncCreateTexture2D is not supported"));
+	return FTexture2DRHIRef();
+}
+
+void FMetalDynamicRHI::RHICopySharedMips(FTexture2DRHIParamRef DestTexture2D,FTexture2DRHIParamRef SrcTexture2D)
+{
+
+}
+
+FTexture2DArrayRHIRef FMetalDynamicRHI::RHICreateTexture2DArray(uint32 SizeX, uint32 SizeY, uint32 SizeZ, uint8 Format, uint32 NumMips, uint32 Flags, FRHIResourceCreateInfo& CreateInfo)
+{
+	return new FMetalTexture2DArray((EPixelFormat)Format, SizeX, SizeY, SizeZ, NumMips, Flags, CreateInfo.BulkData);
+}
+
+FTexture3DRHIRef FMetalDynamicRHI::RHICreateTexture3D(uint32 SizeX, uint32 SizeY, uint32 SizeZ, uint8 Format, uint32 NumMips, uint32 Flags, FRHIResourceCreateInfo& CreateInfo)
+{
+	return new FMetalTexture3D((EPixelFormat)Format, SizeX, SizeY, SizeZ, NumMips, Flags, CreateInfo.BulkData);
+}
+
+void FMetalDynamicRHI::RHIGetResourceInfo(FTextureRHIParamRef Ref, FRHIResourceInfo& OutInfo)
+{
+}
+
+void FMetalDynamicRHI::RHIGenerateMips(FTextureRHIParamRef SourceSurfaceRHI)
+{
+
+}
+
+FTexture2DRHIRef FMetalDynamicRHI::RHIAsyncReallocateTexture2D(FTexture2DRHIParamRef OldTextureRHI, int32 NewMipCount, int32 NewSizeX, int32 NewSizeY, FThreadSafeCounter* RequestStatus)
+{
+	DYNAMIC_CAST_METGALRESOURCE(Texture2D,OldTexture);
+
+	return NULL;
+}
+
+ETextureReallocationStatus FMetalDynamicRHI::RHIFinalizeAsyncReallocateTexture2D( FTexture2DRHIParamRef Texture2D, bool bBlockUntilCompleted )
+{
+	return TexRealloc_Failed;
+}
+
+ETextureReallocationStatus FMetalDynamicRHI::RHICancelAsyncReallocateTexture2D( FTexture2DRHIParamRef Texture2D, bool bBlockUntilCompleted )
+{
+	return TexRealloc_Failed;
+}
+
+
+void* FMetalDynamicRHI::RHILockTexture2D(FTexture2DRHIParamRef TextureRHI,uint32 MipIndex,EResourceLockMode LockMode,uint32& DestStride,bool bLockWithinMiptail)
+{
+	DYNAMIC_CAST_METGALRESOURCE(Texture2D,Texture);
+	return Texture->Surface.Lock(MipIndex, 0, LockMode, DestStride);
+}
+
+void FMetalDynamicRHI::RHIUnlockTexture2D(FTexture2DRHIParamRef TextureRHI,uint32 MipIndex,bool bLockWithinMiptail)
+{
+	DYNAMIC_CAST_METGALRESOURCE(Texture2D,Texture);
+	Texture->Surface.Unlock(MipIndex, 0);
+}
+
+void* FMetalDynamicRHI::RHILockTexture2DArray(FTexture2DArrayRHIParamRef TextureRHI,uint32 TextureIndex,uint32 MipIndex,EResourceLockMode LockMode,uint32& DestStride,bool bLockWithinMiptail)
+{
+	DYNAMIC_CAST_METGALRESOURCE(Texture2DArray,Texture);
+	return Texture->Surface.Lock(MipIndex, TextureIndex, LockMode, DestStride);
+}
+
+void FMetalDynamicRHI::RHIUnlockTexture2DArray(FTexture2DArrayRHIParamRef TextureRHI,uint32 TextureIndex,uint32 MipIndex,bool bLockWithinMiptail)
+{
+	DYNAMIC_CAST_METGALRESOURCE(Texture2DArray,Texture);
+	Texture->Surface.Unlock(MipIndex, TextureIndex);
+}
+
+void FMetalDynamicRHI::RHIUpdateTexture2D(FTexture2DRHIParamRef TextureRHI, uint32 MipIndex, const struct FUpdateTextureRegion2D& UpdateRegion, uint32 SourcePitch, const uint8* SourceData)
+{	
+	DYNAMIC_CAST_METGALRESOURCE(Texture3D,Texture);
+
+}
+
+void FMetalDynamicRHI::RHIUpdateTexture3D(FTexture3DRHIParamRef TextureRHI,uint32 MipIndex,const FUpdateTextureRegion3D& UpdateRegion,uint32 SourceRowPitch,uint32 SourceDepthPitch,const uint8* SourceData)
+{	
+	DYNAMIC_CAST_METGALRESOURCE(Texture3D,Texture);
+
+}
+
+
+/*-----------------------------------------------------------------------------
+	Cubemap texture support.
+-----------------------------------------------------------------------------*/
+FTextureCubeRHIRef FMetalDynamicRHI::RHICreateTextureCube(uint32 Size, uint8 Format, uint32 NumMips, uint32 Flags, FRHIResourceCreateInfo& CreateInfo)
+{
+	return new FMetalTextureCube((EPixelFormat)Format, Size, false, 1, NumMips, Flags, CreateInfo.BulkData);
+}
+
+FTextureCubeRHIRef FMetalDynamicRHI::RHICreateTextureCubeArray(uint32 Size, uint32 ArraySize, uint8 Format, uint32 NumMips, uint32 Flags, FRHIResourceCreateInfo& CreateInfo)
+{
+	return new FMetalTextureCube((EPixelFormat)Format, Size, true, ArraySize, NumMips, Flags, CreateInfo.BulkData);
+}
+
+void* FMetalDynamicRHI::RHILockTextureCubeFace(FTextureCubeRHIParamRef TextureCubeRHI,uint32 FaceIndex,uint32 ArrayIndex,uint32 MipIndex,EResourceLockMode LockMode,uint32& DestStride,bool bLockWithinMiptail)
+{
+	DYNAMIC_CAST_METGALRESOURCE(TextureCube,TextureCube);
+	uint32 MetalFace = GetMetalCubeFace((ECubeFace)FaceIndex);
+	return TextureCube->Surface.Lock(MipIndex, FaceIndex + 6 * ArrayIndex, LockMode, DestStride);
+}
+
+void FMetalDynamicRHI::RHIUnlockTextureCubeFace(FTextureCubeRHIParamRef TextureCubeRHI,uint32 FaceIndex,uint32 ArrayIndex,uint32 MipIndex,bool bLockWithinMiptail)
+{
+	DYNAMIC_CAST_METGALRESOURCE(TextureCube,TextureCube);
+	uint32 MetalFace = GetMetalCubeFace((ECubeFace)FaceIndex);
+	TextureCube->Surface.Unlock(MipIndex, FaceIndex + ArrayIndex * 6);
+}
+
+
+
+
+FTextureReferenceRHIRef FMetalDynamicRHI::RHICreateTextureReference()
+{
+	return new FMetalTextureReference();
+}
+
+void FMetalDynamicRHI::RHIUpdateTextureReference(FTextureReferenceRHIParamRef TextureRefRHI, FTextureRHIParamRef NewTextureRHI)
+{
+	FMetalTextureReference* TextureRef = (FMetalTextureReference*)TextureRefRHI;
+	if (TextureRef)
+	{
+		TextureRef->SetReferencedTexture(NewTextureRHI);
+	}
+}
+
+
+void FMetalDynamicRHI::RHIBindDebugLabelName(FTextureRHIParamRef TextureRHI, const TCHAR* Name)
+{
+}
+
+void FMetalDynamicRHI::RHIVirtualTextureSetFirstMipInMemory(FTexture2DRHIParamRef TextureRHI, uint32 FirstMip)
+{
+	NOT_SUPPORTED("RHIVirtualTextureSetFirstMipInMemory");
+}
+
+void FMetalDynamicRHI::RHIVirtualTextureSetFirstMipVisible(FTexture2DRHIParamRef TextureRHI, uint32 FirstMip)
+{
+	NOT_SUPPORTED("RHIVirtualTextureSetFirstMipVisible");
+}
diff --git a/Engine/Source/Runtime/IOS/MetalRHI/Private/MetalUAV.cpp b/Engine/Source/Runtime/IOS/MetalRHI/Private/MetalUAV.cpp
new file mode 100644
--- /dev/null
+++ b/Engine/Source/Runtime/IOS/MetalRHI/Private/MetalUAV.cpp
@@ -0,0 +1,84 @@
+// Copyright 1998-2014 Epic Games, Inc. All Rights Reserved.
+
+
+#include "MetalRHIPrivate.h"
+
+FMetalShaderResourceView::~FMetalShaderResourceView()
+{
+	SourceVertexBuffer = NULL;
+	SourceTexture = NULL;
+}
+
+
+
+FUnorderedAccessViewRHIRef FMetalDynamicRHI::RHICreateUnorderedAccessView(FStructuredBufferRHIParamRef StructuredBufferRHI, bool bUseUAVCounter, bool bAppendBuffer)
+{
+	DYNAMIC_CAST_METGALRESOURCE(StructuredBuffer,StructuredBuffer);
+
+	// create the UAV buffer to point to the structured buffer's memory
+	FMetalUnorderedAccessView* UAV = new FMetalUnorderedAccessView;
+	UAV->SourceStructuredBuffer = StructuredBuffer;
+
+	return UAV;
+}
+
+FUnorderedAccessViewRHIRef FMetalDynamicRHI::RHICreateUnorderedAccessView(FTextureRHIParamRef TextureRHI)
+{
+	FMetalSurface& Surface = GetMetalSurfaceFromRHITexture(TextureRHI);
+
+	// create the UAV buffer to point to the structured buffer's memory
+	FMetalUnorderedAccessView* UAV = new FMetalUnorderedAccessView;
+	UAV->SourceTexture = (FRHITexture*)TextureRHI;
+
+	return UAV;
+}
+
+FUnorderedAccessViewRHIRef FMetalDynamicRHI::RHICreateUnorderedAccessView(FVertexBufferRHIParamRef VertexBufferRHI, uint8 Format)
+{
+	DYNAMIC_CAST_METGALRESOURCE(VertexBuffer, VertexBuffer);
+
+	// create the UAV buffer to point to the structured buffer's memory
+	FMetalUnorderedAccessView* UAV = new FMetalUnorderedAccessView;
+	UAV->SourceVertexBuffer = VertexBuffer;
+
+	return UAV;
+}
+
+FShaderResourceViewRHIRef FMetalDynamicRHI::RHICreateShaderResourceView(FStructuredBufferRHIParamRef StructuredBufferRHI)
+{
+	DYNAMIC_CAST_METGALRESOURCE(StructuredBuffer, StructuredBuffer);
+
+	FMetalShaderResourceView* SRV = new FMetalShaderResourceView;
+	return SRV;
+}
+
+FShaderResourceViewRHIRef FMetalDynamicRHI::RHICreateShaderResourceView(FVertexBufferRHIParamRef VertexBufferRHI, uint32 Stride, uint8 Format)
+{
+	DYNAMIC_CAST_METGALRESOURCE(VertexBuffer, VertexBuffer);
+
+	FMetalShaderResourceView* SRV = new FMetalShaderResourceView;
+	SRV->SourceVertexBuffer = VertexBuffer;
+	return SRV;
+}
+
+FShaderResourceViewRHIRef FMetalDynamicRHI::RHICreateShaderResourceView(FTexture2DRHIParamRef Texture2DRHI, uint8 MipLevel)
+{
+	FMetalShaderResourceView* SRV = new FMetalShaderResourceView;
+	SRV->SourceTexture = (FRHITexture*)Texture2DRHI;
+	return SRV;
+}
+
+FShaderResourceViewRHIRef FMetalDynamicRHI::RHICreateShaderResourceView(FTexture2DRHIParamRef Texture2DRHI, uint8 MipLevel, uint8 NumMipLevels, uint8 Format)
+{
+	FMetalShaderResourceView* SRV = new FMetalShaderResourceView;
+	SRV->SourceTexture = (FRHITexture*)Texture2DRHI;
+	return SRV;
+}
+
+void FMetalDynamicRHI::RHIClearUAV(FUnorderedAccessViewRHIParamRef UnorderedAccessViewRHI, const uint32* Values)
+{
+	DYNAMIC_CAST_METGALRESOURCE(UnorderedAccessView, UnorderedAccessView);
+
+}
+
+
diff --git a/Engine/Source/Runtime/IOS/MetalRHI/Private/MetalUniformBuffer.cpp b/Engine/Source/Runtime/IOS/MetalRHI/Private/MetalUniformBuffer.cpp
new file mode 100644
--- /dev/null
+++ b/Engine/Source/Runtime/IOS/MetalRHI/Private/MetalUniformBuffer.cpp
@@ -0,0 +1,236 @@
+// Copyright 1998-2014 Epic Games, Inc. All Rights Reserved.
+
+/*=============================================================================
+	MetalConstantBuffer.cpp: Metal Constant buffer implementation.
+=============================================================================*/
+
+#include "MetalRHIPrivate.h"
+
+
+#define NUM_POOL_BUCKETS 45
+#define NUM_SAFE_FRAMES 3
+
+static const uint32 RequestedUniformBufferSizeBuckets[NUM_POOL_BUCKETS] =
+{
+	16, 32, 48, 64, 80, 96, 112, 128,	// 16-byte increments
+	160, 192, 224, 256,			// 32-byte increments
+	320, 384, 448, 512,			// 64-byte increments
+	640, 768, 896, 1024,			// 128-byte increments
+	1280, 1536, 1792, 2048,		// 256-byte increments
+	2560, 3072, 3584, 4096,		// 512-byte increments
+	5120, 6144, 7168, 8192,		// 1024-byte increments
+	10240, 12288, 14336, 16384,	// 2048-byte increments
+	20480, 24576, 28672, 32768,	// 4096-byte increments
+	40960, 49152, 57344, 65536,	// 8192-byte increments
+
+	// 65536 is current max uniform buffer size for Mac OS X.
+
+	0xFFFFFFFF
+};
+
+// Maps desired size buckets to alignment actually 
+static TArray<uint32> UniformBufferSizeBuckets;
+
+static uint32 GetUBPoolSize()
+{
+	return 512 * 1024;
+}
+
+// Convert bucket sizes to be compatible with present device
+static void RemapBuckets()
+{
+	const uint32 Alignment = 256;
+
+	for (int32 Count = 0; Count < NUM_POOL_BUCKETS; Count++)
+	{
+		uint32 AlignedSize = ((RequestedUniformBufferSizeBuckets[Count] + Alignment - 1) / Alignment ) * Alignment;
+		if (!UniformBufferSizeBuckets.Contains(AlignedSize))
+		{
+			UniformBufferSizeBuckets.Push(AlignedSize);
+		}
+	}
+}
+
+static uint32 GetPoolBucketIndex(uint32 NumBytes)
+{
+	if (UniformBufferSizeBuckets.Num() == 0)
+	{
+		RemapBuckets();
+	}
+
+	unsigned long Lower = 0;
+	unsigned long Upper = UniformBufferSizeBuckets.Num();
+
+	do
+	{
+		unsigned long Middle = (Upper + Lower) >> 1;
+		if (NumBytes <= UniformBufferSizeBuckets[Middle - 1])
+		{
+			Upper = Middle;
+		}
+		else
+		{
+			Lower = Middle;
+		}
+	}
+	while (Upper - Lower > 1);
+
+	check(NumBytes <= UniformBufferSizeBuckets[Lower]);
+	check((Lower == 0) || (NumBytes > UniformBufferSizeBuckets[Lower - 1]));
+
+	return Lower;
+}
+
+
+// Describes a uniform buffer in the free pool.
+struct FPooledUniformBuffer
+{
+	id<MTLBuffer> Buffer;
+	uint32 CreatedSize;
+	uint32 FrameFreed;
+	uint32 Offset;
+};
+
+// Pool of free uniform buffers, indexed by bucket for constant size search time.
+TArray<FPooledUniformBuffer> UniformBufferPool[NUM_POOL_BUCKETS];
+
+// Uniform buffers that have been freed more recently than NumSafeFrames ago.
+TArray<FPooledUniformBuffer> SafeUniformBufferPools[NUM_SAFE_FRAMES][NUM_POOL_BUCKETS];
+
+// Does per-frame global updating for the uniform buffer pool.
+void InitFrame_UniformBufferPoolCleanup()
+{
+	check(IsInRenderingThread());
+
+	SCOPE_CYCLE_COUNTER(STAT_MetalUniformBufferCleanupTime);
+
+	// Index of the bucket that is now old enough to be reused
+	const int32 SafeFrameIndex = GFrameNumberRenderThread % NUM_SAFE_FRAMES;
+
+	// Merge the bucket into the free pool array
+	for (int32 BucketIndex = 0; BucketIndex < NUM_POOL_BUCKETS; BucketIndex++)
+	{
+		UniformBufferPool[BucketIndex].Append(SafeUniformBufferPools[SafeFrameIndex][BucketIndex]);
+		SafeUniformBufferPools[SafeFrameIndex][BucketIndex].Reset();
+	}
+}
+
+struct TUBPoolBuffer
+{
+	id<MTLBuffer> Buffer;
+	uint32 ConsumedSpace;
+	uint32 AllocatedSpace;
+};
+
+TArray<TUBPoolBuffer> UBPool;
+
+void AddNewlyFreedBufferToUniformBufferPool(id<MTLBuffer> Buffer, uint32 Offset, uint32 Size)
+{
+	check(Buffer);
+	FPooledUniformBuffer NewEntry;
+	NewEntry.Buffer = Buffer;
+	NewEntry.FrameFreed = GFrameNumberRenderThread;
+	NewEntry.CreatedSize = Size;
+	NewEntry.Offset = Offset;
+
+	// Add to this frame's array of free uniform buffers
+	const int32 SafeFrameIndex = (GFrameNumberRenderThread - 1) % NUM_SAFE_FRAMES;
+	const uint32 BucketIndex = GetPoolBucketIndex(Size);
+
+	SafeUniformBufferPools[SafeFrameIndex][BucketIndex].Add(NewEntry);
+	INC_DWORD_STAT(STAT_MetalNumFreeUniformBuffers);
+	INC_MEMORY_STAT_BY(STAT_MetalFreeUniformBufferMemory, Buffer.length);
+}
+
+
+id<MTLBuffer> SuballocateUB(uint32 Size, uint32& OutOffset)
+{
+	check(Size <= GetUBPoolSize());
+
+	// Find space in previously allocated pool buffers
+	for ( int32 Buffer = 0; Buffer < UBPool.Num(); Buffer++)
+	{
+		TUBPoolBuffer &Pool = UBPool[Buffer];
+		if ( Size < (Pool.AllocatedSpace - Pool.ConsumedSpace))
+		{
+			OutOffset = Pool.ConsumedSpace;
+			Pool.ConsumedSpace += Size;
+			return Pool.Buffer;
+		}
+	}
+
+	// No space was found to use, create a new Pool buffer
+	uint32 TotalSize = GetUBPoolSize();
+	//NSLog(@"New Metal Buffer Size %d", TotalSize);
+	id<MTLBuffer> Buffer = [FMetalManager::GetDevice() newBufferWithLength:TotalSize options:BUFFER_CACHE_MODE];
+	TRACK_OBJECT(Buffer);
+
+	OutOffset = 0;
+
+	TUBPoolBuffer Pool;
+	Pool.Buffer = Buffer;
+	Pool.ConsumedSpace = Size;
+	Pool.AllocatedSpace = GetUBPoolSize();
+	UBPool.Push(Pool);
+
+	return Buffer;
+}
+
+
+FMetalUniformBuffer::FMetalUniformBuffer(const void* Contents, const FRHIUniformBufferLayout& Layout, EUniformBufferUsage Usage)
+	: FRHIUniformBuffer(Layout)
+	, Buffer(nil)
+	, Offset(0)
+{
+	// for single use buffers, allocate from the ring buffer to avoid thrashing memory
+	if (Usage == UniformBuffer_SingleDraw)
+	{
+		// use a bit of the ring buffer
+		Offset = FMetalManager::Get()->AllocateFromRingBuffer(Layout.ConstantBufferSize);
+		Buffer = FMetalManager::Get()->GetRingBuffer();
+	}
+	else
+	{
+		// Find the appropriate bucket based on size
+		const uint32 BucketIndex = GetPoolBucketIndex(Layout.ConstantBufferSize);
+		TArray<FPooledUniformBuffer>& PoolBucket = UniformBufferPool[BucketIndex];
+		if (PoolBucket.Num() > 0)
+		{
+			// Reuse the last entry in this size bucket
+			FPooledUniformBuffer FreeBufferEntry = PoolBucket.Pop();
+			DEC_DWORD_STAT(STAT_MetalNumFreeUniformBuffers);
+			DEC_MEMORY_STAT_BY(STAT_MetalFreeUniformBufferMemory, FreeBufferEntry.CreatedSize);
+
+			// reuse the one
+			Buffer = FreeBufferEntry.Buffer;
+			Offset = FreeBufferEntry.Offset;
+		}
+		else
+		{
+			// Nothing usable was found in the free pool, create a new uniform buffer (full size, not NumBytes)
+			uint32 BufferSize = UniformBufferSizeBuckets[BucketIndex];
+			Buffer = SuballocateUB(BufferSize, Offset);
+		}
+	}
+
+	// copy the contents
+	FMemory::Memcpy(((uint8*)[Buffer contents]) + Offset, Contents, Layout.ConstantBufferSize);
+
+	// @todo metal: Add support for resource tables!
+	check(Layout.Resources.Num() == 0);
+}
+
+FMetalUniformBuffer::~FMetalUniformBuffer()
+{
+	// don't need to free the ring buffer!
+	if (Buffer != FMetalManager::Get()->GetRingBuffer())
+	{
+		AddNewlyFreedBufferToUniformBufferPool(Buffer, Offset, GetSize());
+		//	FMetalManager::ReleaseObject(Buffer);
+	}
+}
+
+FUniformBufferRHIRef FMetalDynamicRHI::RHICreateUniformBuffer(const void* Contents, const FRHIUniformBufferLayout& Layout, EUniformBufferUsage Usage)
+{
+	return new FMetalUniformBuffer(Contents, Layout, Usage);
+}
diff --git a/Engine/Source/Runtime/IOS/MetalRHI/Private/MetalVertexBuffer.cpp b/Engine/Source/Runtime/IOS/MetalRHI/Private/MetalVertexBuffer.cpp
new file mode 100644
--- /dev/null
+++ b/Engine/Source/Runtime/IOS/MetalRHI/Private/MetalVertexBuffer.cpp
@@ -0,0 +1,98 @@
+// Copyright 1998-2014 Epic Games, Inc. All Rights Reserved.
+
+/*=============================================================================
+	MetalVertexBuffer.cpp: Metal vertex buffer RHI implementation.
+=============================================================================*/
+
+#include "MetalRHIPrivate.h"
+
+
+FMetalVertexBuffer::FMetalVertexBuffer(uint32 InSize, uint32 InUsage)
+	: FRHIVertexBuffer(InSize, InUsage)
+	, ZeroStrideElementSize((InUsage & BUF_ZeroStride) ? InSize : 0)
+{
+	// volatile buffers can come right from the ring buffer
+	if (InUsage & BUF_Volatile)
+	{
+		Buffer = FMetalManager::Get()->GetRingBuffer();
+		// retaining to ease the destructor logic
+		[Buffer retain];
+	}
+	else
+	{
+		// reduce some churn in the ZeroStride buffer hack
+		if (InSize <= 16)
+		{
+			InSize = 128 * 1024;
+		}
+//		NSLog(@"Allocating %d byte vertex buffer", InSize);
+		
+		Buffer = [FMetalManager::GetDevice() newBufferWithLength:InSize options:BUFFER_CACHE_MODE];
+		Offset = 0;
+	}
+	TRACK_OBJECT(Buffer);
+}
+
+FMetalVertexBuffer::~FMetalVertexBuffer()
+{
+	FMetalManager::ReleaseObject(Buffer);
+}
+
+
+void* FMetalVertexBuffer::Lock(EResourceLockMode LockMode, uint32 Size)
+{
+	if (GetUsage() & BUF_Volatile)
+	{
+		Offset = FMetalManager::Get()->AllocateFromRingBuffer(Size);
+	}
+	return ((uint8*)[Buffer contents]) + Offset;
+}
+
+void FMetalVertexBuffer::Unlock()
+{
+	
+}
+
+FVertexBufferRHIRef FMetalDynamicRHI::RHICreateVertexBuffer(uint32 Size, uint32 InUsage, FRHIResourceCreateInfo& CreateInfo)
+{
+	// make the RHI object, which will allocate memory
+	FMetalVertexBuffer* VertexBuffer = new FMetalVertexBuffer(Size, InUsage);
+
+	if (CreateInfo.ResourceArray)
+	{
+		check(Size == CreateInfo.ResourceArray->GetResourceDataSize());
+
+		// make a buffer usable by CPU
+		void* Buffer = RHILockVertexBuffer(VertexBuffer, 0, Size, RLM_WriteOnly);
+
+		// copy the contents of the given data into the buffer
+		FMemory::Memcpy(Buffer, CreateInfo.ResourceArray->GetResourceData(), Size);
+
+		RHIUnlockVertexBuffer(VertexBuffer);
+
+		// Discard the resource array's contents.
+		CreateInfo.ResourceArray->Discard();
+	}
+
+	return VertexBuffer;
+}
+
+void* FMetalDynamicRHI::RHILockVertexBuffer(FVertexBufferRHIParamRef VertexBufferRHI, uint32 Offset, uint32 Size, EResourceLockMode LockMode)
+{
+	DYNAMIC_CAST_METGALRESOURCE(VertexBuffer,VertexBuffer);
+
+	// default to vertex buffer memory
+	return (uint8*)VertexBuffer->Lock(LockMode, Size) + Offset;
+}
+
+void FMetalDynamicRHI::RHIUnlockVertexBuffer(FVertexBufferRHIParamRef VertexBufferRHI)
+{
+	DYNAMIC_CAST_METGALRESOURCE(VertexBuffer,VertexBuffer);
+
+	VertexBuffer->Unlock();
+}
+
+void FMetalDynamicRHI::RHICopyVertexBuffer(FVertexBufferRHIParamRef SourceBufferRHI,FVertexBufferRHIParamRef DestBufferRHI)
+{
+	NOT_SUPPORTED("RHICopyVertexBuffer");
+}
\ No newline at end of file
diff --git a/Engine/Source/Runtime/IOS/MetalRHI/Private/MetalVertexDeclaration.cpp b/Engine/Source/Runtime/IOS/MetalRHI/Private/MetalVertexDeclaration.cpp
new file mode 100644
--- /dev/null
+++ b/Engine/Source/Runtime/IOS/MetalRHI/Private/MetalVertexDeclaration.cpp
@@ -0,0 +1,100 @@
+// Copyright 1998-2014 Epic Games, Inc. All Rights Reserved.
+
+/*=============================================================================
+	MetalVertexDeclaration.cpp: Metal vertex declaration RHI implementation.
+=============================================================================*/
+
+#include "MetalRHIPrivate.h"
+
+static MTLVertexFormat TranslateElementTypeToMTLType(EVertexElementType Type)
+{
+	switch (Type)
+	{
+	case VET_Float1:		return MTLVertexFormatFloat;
+	case VET_Float2:		return MTLVertexFormatFloat2;
+	case VET_Float3:		return MTLVertexFormatFloat3;
+	case VET_Float4:		return MTLVertexFormatFloat4;
+	case VET_PackedNormal:	return MTLVertexFormatUChar4Normalized;
+	case VET_UByte4:		return MTLVertexFormatUChar4;
+	case VET_UByte4N:		return MTLVertexFormatUChar4Normalized;
+	case VET_Color:			return MTLVertexFormatUChar4Normalized;
+	case VET_Short2:		return MTLVertexFormatShort2;
+	case VET_Short4:		return MTLVertexFormatShort4;
+	case VET_Short2N:		return MTLVertexFormatShort2Normalized;
+	case VET_Half2:			return MTLVertexFormatHalf2;
+	default:				return MTLVertexFormatFloat4;
+	};
+}
+
+FMetalVertexDeclaration::FMetalVertexDeclaration(const FVertexDeclarationElementList& InElements)
+	: Elements(InElements)
+	, Layout(nil)
+{
+	GenerateLayout(InElements);
+}
+
+FMetalVertexDeclaration::~FMetalVertexDeclaration()
+{
+	FMetalManager::ReleaseObject(Layout);
+}
+
+static TMap<uint32, FVertexDeclarationRHIRef> GVertexDeclarationCache;
+
+FVertexDeclarationRHIRef FMetalDynamicRHI::RHICreateVertexDeclaration(const FVertexDeclarationElementList& Elements)
+{
+	uint32 Key = FCrc::MemCrc32(Elements.GetData(), Elements.Num() * sizeof(FVertexElement));
+	// look up an existing declaration
+	FVertexDeclarationRHIRef* VertexDeclarationRefPtr = GVertexDeclarationCache.Find(Key);
+	if (VertexDeclarationRefPtr == NULL)
+	{
+//		NSLog(@"VertDecl Key: %x", Key);
+
+		// create and add to the cache if it doesn't exist.
+		VertexDeclarationRefPtr = &GVertexDeclarationCache.Add(Key, new FMetalVertexDeclaration(Elements));
+	}
+
+	return *VertexDeclarationRefPtr;
+}
+
+
+void FMetalVertexDeclaration::GenerateLayout(const FVertexDeclarationElementList& Elements)
+{
+	Layout = [[MTLVertexDescriptor alloc] init];
+	TRACK_OBJECT(Layout);
+
+	TMap<uint32, uint32> BufferStrides;
+	for (uint32 ElementIndex = 0; ElementIndex < Elements.Num(); ElementIndex++)
+	{
+		const FVertexElement& Element = Elements[ElementIndex];
+		
+		// @todo urban: for zero stride elements, assume a repeated color
+		uint32 Stride = Element.Stride ? Element.Stride : 4;
+
+		// we offset 6 buffers to leave space for uniform buffers
+		uint32 ShaderBufferIndex = UNREAL_TO_METAL_BUFFER_INDEX(Element.StreamIndex);
+
+		// track the buffer stride, making sure all elements with the same buffer have the same stride
+		uint32* ExistingStride = BufferStrides.Find(ShaderBufferIndex);
+		if (ExistingStride == NULL)
+		{
+			BufferStrides.Add(ShaderBufferIndex, Stride);
+		}
+		else
+		{
+			check(Stride == *ExistingStride);
+		}
+
+		[Layout setVertexFormat:TranslateElementTypeToMTLType(Element.Type) offset:Element.Offset vertexBufferIndex:ShaderBufferIndex atAttributeIndex:Element.AttributeIndex];
+	}
+
+	for (auto It = BufferStrides.CreateIterator(); It; ++It)
+	{
+		uint32 Stride = It.Value();
+		if (Stride == 0xFFFF)
+		{
+			NSLog(@"Setting illegal stride - break here if you want to find out why, but this won't break until we try to render with it");
+			Stride = 200;// 16;
+		}
+		[Layout setStride:Stride instanceStepRate:0 atVertexBufferIndex:It.Key()];
+	}
+}
diff --git a/Engine/Source/Runtime/IOS/MetalRHI/Private/MetalViewport.cpp b/Engine/Source/Runtime/IOS/MetalRHI/Private/MetalViewport.cpp
new file mode 100644
--- /dev/null
+++ b/Engine/Source/Runtime/IOS/MetalRHI/Private/MetalViewport.cpp
@@ -0,0 +1,64 @@
+// Copyright 1998-2014 Epic Games, Inc. All Rights Reserved.
+
+/*=============================================================================
+	MetalViewport.cpp: Metal viewport RHI implementation.
+=============================================================================*/
+
+#include "MetalRHIPrivate.h"
+
+
+FMetalViewport::FMetalViewport(void* WindowHandle, uint32 InSizeX,uint32 InSizeY,bool bInIsFullscreen)
+{
+
+}
+
+FMetalViewport::~FMetalViewport()
+{
+}
+
+
+/*=============================================================================
+ *	The following RHI functions must be called from the main thread.
+ *=============================================================================*/
+FViewportRHIRef FMetalDynamicRHI::RHICreateViewport(void* WindowHandle,uint32 SizeX,uint32 SizeY,bool bIsFullscreen)
+{
+	check( IsInGameThread() );
+	return new FMetalViewport(WindowHandle, SizeX, SizeY, bIsFullscreen);
+}
+
+void FMetalDynamicRHI::RHIResizeViewport(FViewportRHIParamRef ViewportRHI,uint32 SizeX,uint32 SizeY,bool bIsFullscreen)
+{
+	check( IsInGameThread() );
+
+	DYNAMIC_CAST_METGALRESOURCE(Viewport,Viewport);
+}
+
+void FMetalDynamicRHI::RHITick( float DeltaTime )
+{
+	check( IsInGameThread() );
+
+}
+
+/*=============================================================================
+ *	Viewport functions.
+ *=============================================================================*/
+
+void FMetalDynamicRHI::RHIBeginDrawingViewport(FViewportRHIParamRef ViewportRHI, FTextureRHIParamRef RenderTargetRHI)
+{
+	FMetalManager::Get()->BeginFrame();
+}
+
+void FMetalDynamicRHI::RHIEndDrawingViewport(FViewportRHIParamRef ViewportRHI,bool bPresent,bool bLockToVsync)
+{
+	FMetalManager::Get()->EndFrame(bPresent);
+}
+
+bool FMetalDynamicRHI::RHIIsDrawingViewport()
+{
+	return true;
+}
+
+FTexture2DRHIRef FMetalDynamicRHI::RHIGetViewportBackBuffer(FViewportRHIParamRef ViewportRHI)
+{
+	return FMetalManager::Get()->GetBackBuffer();
+}
diff --git a/Engine/Source/Runtime/IOS/MetalRHI/Public/MetalRHI.h b/Engine/Source/Runtime/IOS/MetalRHI/Public/MetalRHI.h
new file mode 100644
--- /dev/null
+++ b/Engine/Source/Runtime/IOS/MetalRHI/Public/MetalRHI.h
@@ -0,0 +1,62 @@
+// Copyright 1998-2014 Epic Games, Inc. All Rights Reserved.
+
+/*=============================================================================
+	MetalRHI.h: Public Metal RHI definitions.
+=============================================================================*/
+
+#pragma once 
+
+DECLARE_LOG_CATEGORY_EXTERN(LogMetal, Display, All);
+
+/** This is a macro that casts a dynamically bound RHI reference to the appropriate type. */
+#define DYNAMIC_CAST_METGALRESOURCE(Type,Name) \
+	FMetal##Type* Name = (FMetal##Type*)Name##RHI;
+
+// Metal RHI public headers.
+#include <Metal/Metal.h>
+#include "MetalState.h"
+#include "MetalResources.h"
+#include "MetalViewport.h"
+
+class FMetalManager;
+
+/** The interface which is implemented by the dynamically bound RHI. */
+class FMetalDynamicRHI : public FDynamicRHI
+{
+public:
+
+	/** Initialization constructor. */
+	FMetalDynamicRHI();
+
+	/** Destructor */
+	~FMetalDynamicRHI();
+
+
+	// FDynamicRHI interface.
+	virtual void Init();
+	virtual void PushEvent(const TCHAR* Name);
+	virtual void PopEvent();
+	virtual void Shutdown() {}
+
+	// The RHI methods are defined as virtual functions in URenderHardwareInterface.
+	#define DEFINE_RHIMETHOD(Type,Name,ParameterTypesAndNames,ParameterNames,ReturnStatement,NullImplementation) virtual Type Name ParameterTypesAndNames
+	#include "RHIMethods.h"
+	#undef DEFINE_RHIMETHOD
+
+protected:
+    TGlobalResource<TBoundShaderStateHistory<10000>> BoundShaderStateHistory;
+};
+
+/** Implements the Metal module as a dynamic RHI providing module. */
+class FMetalDynamicRHIModule : public IDynamicRHIModule
+{
+public:
+
+	// IModuleInterface
+	virtual bool SupportsShutdown() { return false; }
+
+	// IDynamicRHIModule
+	virtual bool IsSupported();
+
+	virtual FDynamicRHI* CreateRHI();
+};
diff --git a/Engine/Source/Runtime/IOS/MetalRHI/Public/MetalResources.h b/Engine/Source/Runtime/IOS/MetalRHI/Public/MetalResources.h
new file mode 100644
--- /dev/null
+++ b/Engine/Source/Runtime/IOS/MetalRHI/Public/MetalResources.h
@@ -0,0 +1,421 @@
+// Copyright 1998-2014 Epic Games, Inc. All Rights Reserved.
+
+/*=============================================================================
+	MetalResources.h: Metal resource RHI definitions.
+=============================================================================*/
+
+#pragma once
+
+#include "BoundShaderStateCache.h"
+#include "MetalShaderResources.h"
+
+/** This represents a vertex declaration that hasn't been combined with a specific shader to create a bound shader. */
+class FMetalVertexDeclaration : public FRHIVertexDeclaration
+{
+public:
+
+	/** Initialization constructor. */
+	FMetalVertexDeclaration(const FVertexDeclarationElementList& InElements);
+	~FMetalVertexDeclaration();
+	
+	/** Cached element info array (offset, stream index, etc) */
+	FVertexDeclarationElementList Elements;
+
+	/** This is the layout for the vertex elements */
+ 	MTLVertexDescriptor* Layout;
+
+protected:
+	void GenerateLayout(const FVertexDeclarationElementList& Elements);
+
+};
+
+
+
+/** This represents a vertex shader that hasn't been combined with a specific declaration to create a bound shader. */
+template<typename BaseResourceType /*, typename ShaderType */>
+class TMetalBaseShader : public BaseResourceType, public IRefCountedObject
+{
+public:
+
+	/** Initialization constructor. */
+	TMetalBaseShader()
+	{
+	}
+	TMetalBaseShader(const TArray<uint8>& InCode);
+
+	/** Destructor */
+	virtual ~TMetalBaseShader();
+
+
+	// IRefCountedObject interface.
+	virtual uint32 AddRef() const
+	{
+		return FRefCountedObject::AddRef();
+	}
+	virtual uint32 Release() const
+	{
+		return FRefCountedObject::Release();
+	}
+	virtual uint32 GetRefCount() const
+	{
+		return FRefCountedObject::GetRefCount();
+	}
+	
+	// this is the compiler shader
+	id<MTLFunction> Function;
+
+ 	/** External bindings for this shader. */
+	FMetalShaderBindings Bindings;
+	TArray< TRefCountPtr<FRHIUniformBuffer> > BoundUniformBuffers;
+	// List of memory copies from RHIUniformBuffer to packed uniforms
+	TArray<FMetalUniformBufferCopyInfo> UniformBuffersCopyInfo;
+    
+	TArray<ANSICHAR> GlslCode;
+    NSString* GlslCodeNSString;
+	const ANSICHAR*  GlslCodeString; // make it easier in VS to see shader code in debug mode; points to begin of GlslCode
+};
+
+typedef TMetalBaseShader<FRHIVertexShader> FMetalVertexShader;
+typedef TMetalBaseShader<FRHIPixelShader> FMetalPixelShader;
+typedef TMetalBaseShader<FRHIHullShader> FMetalHullShader;
+typedef TMetalBaseShader<FRHIDomainShader> FMetalDomainShader;
+typedef TMetalBaseShader<FRHIComputeShader> FMetalComputeShader;
+typedef TMetalBaseShader<FRHIGeometryShader> FMetalGeometryShader;
+
+
+/**
+ * Combined shader state and vertex definition for rendering geometry. 
+ * Each unique instance consists of a vertex decl, vertex shader, and pixel shader.
+ */
+class FMetalBoundShaderState : public FRHIBoundShaderState
+{
+public:
+
+	FCachedBoundShaderStateLink CacheLink;
+
+	/** Cached vertex structure */
+	TRefCountPtr<FMetalVertexDeclaration> VertexDeclaration;
+
+	/** Cached shaders */
+	TRefCountPtr<FMetalVertexShader> VertexShader;
+	TRefCountPtr<FMetalPixelShader> PixelShader;
+	TRefCountPtr<FMetalHullShader> HullShader;
+	TRefCountPtr<FMetalDomainShader> DomainShader;
+	TRefCountPtr<FMetalGeometryShader> GeometryShader;
+
+	/** Initialization constructor. */
+	FMetalBoundShaderState(
+		FVertexDeclarationRHIParamRef InVertexDeclarationRHI,
+		FVertexShaderRHIParamRef InVertexShaderRHI,
+		FPixelShaderRHIParamRef InPixelShaderRHI,
+		FHullShaderRHIParamRef InHullShaderRHI,
+		FDomainShaderRHIParamRef InDomainShaderRHI,
+		FGeometryShaderRHIParamRef InGeometryShaderRHI);
+
+	/**
+	 *Destructor
+	 */
+	~FMetalBoundShaderState();
+	
+	/**
+	 * Prepare a pipeline state object for the current state right before drawing
+	 */
+	void PrepareToDraw(const struct FPipelineShadow& Pipeline);
+
+protected:
+	TMap<uint64, id<MTLRenderPipelineState> > PipelineStates;
+};
+
+
+/** Texture/RT wrapper. */
+class FMetalSurface
+{
+public:
+
+	/** 
+	 * Constructor that will create Texture and Color/DepthBuffers as needed
+	 */
+	FMetalSurface(ERHIResourceType ResourceType, EPixelFormat Format, uint32 SizeX, uint32 SizeY, uint32 SizeZ, uint32 NumSamples, bool bArray, uint32 ArraySize, uint32 NumMips, uint32 Flags, FResourceBulkDataInterface* BulkData);
+
+	/**
+	 * Destructor
+	 */
+	~FMetalSurface();
+
+	/**
+	 * Locks one of the texture's mip-maps.
+	 * @param ArrayIndex Index of the texture array/face in the form Index*6+Face
+	 * @return A pointer to the specified texture data.
+	 */
+	void* Lock(uint32 MipIndex, uint32 ArrayIndex, EResourceLockMode LockMode, uint32& DestStride);
+
+	/** Unlocks a previously locked mip-map.
+	 * @param ArrayIndex Index of the texture array/face in the form Index*6+Face
+	 */
+	void Unlock(uint32 MipIndex, uint32 ArrayIndex);
+
+	/**
+	 * Returns how much memory is used by the surface
+	 */
+	uint32 GetMemorySize();
+
+	EPixelFormat PixelFormat;
+	id<MTLTexture> Texture;
+    id<MTLTexture> MSAATexture;
+	uint32 SizeX, SizeY, SizeZ;
+	bool bIsCubemap;
+	
+	uint32 Flags;
+	void* LockedMemory;
+};
+
+class FMetalTexture2D : public FRHITexture2D
+{
+public:
+	/** The surface info */
+	FMetalSurface Surface;
+
+	// Constructor, just calls base and Surface constructor
+	FMetalTexture2D(EPixelFormat Format, uint32 SizeX, uint32 SizeY, uint32 NumMips, uint32 NumSamples, uint32 Flags, FResourceBulkDataInterface* BulkData)
+		: FRHITexture2D(SizeX, SizeY, NumMips, NumSamples, Format, Flags)
+		, Surface(RRT_Texture2D, Format, SizeX, SizeY, 1, NumSamples, /*bArray=*/ false, 1, NumMips, Flags, BulkData)
+	{
+	}
+    
+    
+};
+
+class FMetalTexture2DArray : public FRHITexture2DArray
+{
+public:
+	/** The surface info */
+	FMetalSurface Surface;
+
+	// Constructor, just calls base and Surface constructor
+	FMetalTexture2DArray(EPixelFormat Format, uint32 SizeX, uint32 SizeY, uint32 ArraySize, uint32 NumMips, uint32 Flags, FResourceBulkDataInterface* BulkData)
+		: FRHITexture2DArray(SizeX, SizeY, ArraySize, NumMips, Format, Flags)
+		, Surface(RRT_Texture2DArray, Format, SizeX, SizeY, 1, /*NumSamples=*/1, /*bArray=*/ true, ArraySize, NumMips, Flags, BulkData)
+	{
+	}
+};
+
+class FMetalTexture3D : public FRHITexture3D
+{
+public:
+	/** The surface info */
+	FMetalSurface Surface;
+
+	// Constructor, just calls base and Surface constructor
+	FMetalTexture3D(EPixelFormat Format, uint32 SizeX, uint32 SizeY, uint32 SizeZ, uint32 NumMips, uint32 Flags, FResourceBulkDataInterface* BulkData)
+		: FRHITexture3D(SizeX, SizeY, SizeZ, NumMips, Format, Flags)
+		, Surface(RRT_Texture3D, Format, SizeX, SizeY, SizeZ, /*NumSamples=*/1, /*bArray=*/ false, 1, NumMips, Flags, BulkData)
+	{
+	}
+};
+
+class FMetalTextureCube : public FRHITextureCube
+{
+public:
+	/** The surface info */
+	FMetalSurface Surface;
+
+	// Constructor, just calls base and Surface constructor
+	FMetalTextureCube(EPixelFormat Format, uint32 Size, bool bArray, uint32 ArraySize, uint32 NumMips, uint32 Flags, FResourceBulkDataInterface* BulkData)
+		: FRHITextureCube(Size, NumMips, Format, Flags)
+		, Surface(RRT_TextureCube, Format, Size, Size, 6, /*NumSamples=*/1, bArray, ArraySize, NumMips, Flags, BulkData)
+	{
+	}
+};
+
+/** Given a pointer to a RHI texture that was created by the Metal RHI, returns a pointer to the FMetalTextureBase it encapsulates. */
+FMetalSurface& GetSurfaceFromRHITexture(FRHITexture* Texture);
+
+
+
+
+
+/** Metal occlusion query */
+class FMetalRenderQuery : public FRHIRenderQuery
+{
+public:
+
+	/** Initialization constructor. */
+	FMetalRenderQuery(ERenderQueryType InQueryType);
+
+	~FMetalRenderQuery();
+
+	/**
+	 * Kick off an occlusion test 
+	 */
+	void Begin();
+
+	/**
+	 * Finish up an occlusion test 
+	 */
+	void End();
+
+	// Where in the ring buffer to store the result
+	uint32 Offset;
+
+	// The command buffer this query is being rendered in - used to make sure the command buffer is complete whewn we read the result
+	uint64 CommandBufferIndex;
+};
+
+/** Index buffer resource class that stores stride information. */
+class FMetalIndexBuffer : public FRHIIndexBuffer
+{
+public:
+
+	/** Constructor */
+	FMetalIndexBuffer(uint32 InStride, uint32 InSize, uint32 InUsage);
+	~FMetalIndexBuffer();
+
+	/**
+	 * Prepare a CPU accessible buffer for uploading to GPU memory
+	 */
+	void* Lock(EResourceLockMode LockMode, uint32 Size=0);
+
+	/**
+	 * Prepare a CPU accessible buffer for uploading to GPU memory
+	 */
+	void Unlock();
+	
+	// balsa buffer memory
+	id<MTLBuffer> Buffer;
+
+	// offet into the buffer (for ring buffer usage)
+	uint32 Offset;
+	
+	// 16- or 32-bit
+	MTLIndexType IndexType;
+};
+
+
+/** Vertex buffer resource class that stores usage type. */
+class FMetalVertexBuffer : public FRHIVertexBuffer
+{
+public:
+
+	/** Constructor */
+	FMetalVertexBuffer(uint32 InSize, uint32 InUsage);
+	~FMetalVertexBuffer();
+
+	/**
+	 * Prepare a CPU accessible buffer for uploading to GPU memory
+	 */
+	void* Lock(EResourceLockMode LockMode, uint32 Size=0);
+
+	/**
+	 * Prepare a CPU accessible buffer for uploading to GPU memory
+	 */
+	void Unlock();
+
+	// balsa buffer memory
+	id<MTLBuffer> Buffer;
+
+	// offset into the buffer (for ring buffer usage)
+	uint32 Offset;
+
+	// if the buffer is a zero stride buffer, we need to duplicate and grow on demand, this is the size of one element to copy
+	uint32 ZeroStrideElementSize;
+};
+
+class FMetalUniformBuffer : public FRHIUniformBuffer
+{
+public:
+
+	// Constructor
+	FMetalUniformBuffer(const void* Contents, const FRHIUniformBufferLayout& Layout, EUniformBufferUsage Usage);
+
+	// Destructor 
+	~FMetalUniformBuffer();
+	
+	/** The buffer containing the contents - either a fresh buffer or the ring buffer */
+	id<MTLBuffer> Buffer;
+	
+	/** This offset is used when passing to setVertexBuffer, etc */
+	uint32 Offset;
+};
+
+
+class FMetalStructuredBuffer : public FRHIStructuredBuffer
+{
+public:
+	// Constructor
+	FMetalStructuredBuffer(uint32 Stride, uint32 Size, FResourceArrayInterface* ResourceArray, uint32 InUsage);
+
+	// Destructor
+	~FMetalStructuredBuffer();
+
+};
+
+
+
+class FMetalUnorderedAccessView : public FRHIUnorderedAccessView
+{
+public:
+
+	// the potential resources to refer to with the UAV object
+	TRefCountPtr<FMetalStructuredBuffer> SourceStructuredBuffer;
+	TRefCountPtr<FMetalVertexBuffer> SourceVertexBuffer;
+	TRefCountPtr<FRHITexture> SourceTexture;
+};
+
+
+class FMetalShaderResourceView : public FRHIShaderResourceView
+{
+public:
+
+	// The vertex buffer this SRV comes from (can be null)
+	TRefCountPtr<FMetalVertexBuffer> SourceVertexBuffer;
+
+	// The texture that this SRV come from
+	TRefCountPtr<FRHITexture> SourceTexture;
+
+	~FMetalShaderResourceView();
+};
+
+class FMetalShaderParameterCache
+{
+public:
+	/** Constructor. */
+	FMetalShaderParameterCache();
+
+	/** Destructor. */
+	~FMetalShaderParameterCache();
+
+	void InitializeResources(int32 UniformArraySize);
+
+	/**
+	 * Marks all uniform arrays as dirty.
+	 */
+	void MarkAllDirty();
+
+	/**
+	 * Sets values directly into the packed uniform array
+	 */
+	void Set(uint32 BufferIndex, uint32 ByteOffset, uint32 NumBytes, const void* NewValues);
+
+	/**
+	 * Commit shader parameters to the currently bound program.
+	 */
+	void CommitPackedGlobals(int32 Stage, const FMetalShaderBindings& Bindings);
+	void CommitPackedUniformBuffers(TRefCountPtr<FMetalBoundShaderState> BoundShaderState, int32 Stage, const TArray< TRefCountPtr<FRHIUniformBuffer> >& UniformBuffers, const TArray<FMetalUniformBufferCopyInfo>& UniformBuffersCopyInfo);
+
+private:
+	/** CPU memory block for storing uniform values. */
+	uint8* PackedGlobalUniforms[METAL_PACKED_TYPEINDEX_MAX];
+
+	struct FRange
+	{
+		uint32	LowVector;
+		uint32	HighVector;
+	};
+	/** Dirty ranges for each uniform array. */
+	FRange	PackedGlobalUniformDirty[METAL_PACKED_TYPEINDEX_MAX];
+
+	/** Scratch CPU memory block for uploading packed uniforms. */
+	uint8* PackedUniformsScratch[METAL_PACKED_TYPEINDEX_MAX];
+
+	int32 GlobalUniformArraySize;
+};
diff --git a/Engine/Source/Runtime/IOS/MetalRHI/Public/MetalShaderResources.h b/Engine/Source/Runtime/IOS/MetalRHI/Public/MetalShaderResources.h
new file mode 100644
--- /dev/null
+++ b/Engine/Source/Runtime/IOS/MetalRHI/Public/MetalShaderResources.h
@@ -0,0 +1,218 @@
+// Copyright 1998-2014 Epic Games, Inc. All Rights Reserved.
+
+/*=============================================================================
+	MetalShaderResources.h: Metal shader resource RHI definitions.
+=============================================================================*/
+
+#pragma once
+
+enum
+{
+	METAL_SHADER_STAGE_VERTEX = 0,
+	METAL_SHADER_STAGE_PIXEL,
+	METAL_SHADER_STAGE_GEOMETRY,
+	METAL_SHADER_STAGE_HULL,
+	METAL_SHADER_STAGE_DOMAIN,
+	METAL_NUM_NON_COMPUTE_SHADER_STAGES,
+	METAL_SHADER_STAGE_COMPUTE = METAL_NUM_NON_COMPUTE_SHADER_STAGES,
+	METAL_NUM_SHADER_STAGES
+};
+
+/**
+* Shader related constants.
+*/
+enum
+{
+	METAL_PACKED_TYPENAME_HIGHP = 'h',	// Make sure these enums match hlslcc
+	METAL_PACKED_TYPENAME_MEDIUMP = 'm',
+	METAL_PACKED_TYPENAME_LOWP = 'l',
+	METAL_PACKED_TYPENAME_INT = 'i',
+	METAL_PACKED_TYPENAME_UINT = 'u',
+	METAL_PACKED_TYPENAME_SAMPLER = 's',
+	METAL_PACKED_TYPENAME_IMAGE = 'g',
+
+	METAL_PACKED_TYPEINDEX_HIGHP = 0,
+	METAL_PACKED_TYPEINDEX_MEDIUMP = 1,
+	METAL_PACKED_TYPEINDEX_LOWP = 2,
+	METAL_PACKED_TYPEINDEX_INT = 3,
+	METAL_PACKED_TYPEINDEX_UINT = 4,
+	METAL_PACKED_TYPEINDEX_MAX = 5,
+
+	METAL_MAX_UNIFORM_BUFFER_BINDINGS = 12,	// @todo-mobile: Remove me
+	METAL_FIRST_UNIFORM_BUFFER = 0,			// @todo-mobile: Remove me
+	METAL_MAX_COMPUTE_STAGE_UAV_UNITS = 8,	// @todo-mobile: Remove me
+	METAL_UAV_NOT_SUPPORTED_FOR_GRAPHICS_UNIT = -1, // for now, only CS supports UAVs/ images
+};
+
+static FORCEINLINE uint8 MetalPackedTypeNameToTypeIndex(uint8 ArrayName)
+{
+	switch (ArrayName)
+	{
+	case METAL_PACKED_TYPENAME_HIGHP:		return METAL_PACKED_TYPEINDEX_HIGHP;
+	case METAL_PACKED_TYPENAME_MEDIUMP:	return METAL_PACKED_TYPEINDEX_MEDIUMP;
+	case METAL_PACKED_TYPENAME_LOWP:		return METAL_PACKED_TYPEINDEX_LOWP;
+	case METAL_PACKED_TYPENAME_INT:		return METAL_PACKED_TYPEINDEX_INT;
+	case METAL_PACKED_TYPENAME_UINT:		return METAL_PACKED_TYPEINDEX_UINT;
+	}
+	check(0);
+	return 0;
+}
+
+struct FMetalPackedArrayInfo
+{
+	uint16	Size;		// Bytes
+	uint8	TypeName;	// OGL_PACKED_TYPENAME
+	uint8	TypeIndex;	// OGL_PACKED_TYPE
+};
+
+inline FArchive& operator<<(FArchive& Ar, FMetalPackedArrayInfo& Info)
+{
+	Ar << Info.Size;
+	Ar << Info.TypeName;
+	Ar << Info.TypeIndex;
+	return Ar;
+}
+
+
+struct FMetalShaderBindings
+{
+	TArray<TArray<FMetalPackedArrayInfo>>	PackedUniformBuffers;
+	TArray<FMetalPackedArrayInfo>			PackedGlobalArrays;
+
+	uint16	InOutMask;
+	uint8	NumSamplers;
+	uint8	NumUniformBuffers;
+	uint8	NumUAVs;
+	bool	bHasRegularUniformBuffers;
+
+	FMetalShaderBindings() :
+		InOutMask(0),
+		NumSamplers(0),
+		NumUniformBuffers(0),
+		NumUAVs(0),
+		bHasRegularUniformBuffers(false)
+	{
+	}
+
+	friend bool operator==(const FMetalShaderBindings &A, const FMetalShaderBindings &B)
+	{
+		bool bEqual = true;
+
+		bEqual &= A.InOutMask == B.InOutMask;
+		bEqual &= A.NumSamplers == B.NumSamplers;
+		bEqual &= A.NumUniformBuffers == B.NumUniformBuffers;
+		bEqual &= A.NumUAVs == B.NumUAVs;
+		bEqual &= A.bHasRegularUniformBuffers == B.bHasRegularUniformBuffers;
+		bEqual &= A.PackedGlobalArrays.Num() == B.PackedGlobalArrays.Num();
+		bEqual &= A.PackedUniformBuffers.Num() == B.PackedUniformBuffers.Num();
+
+		if (!bEqual)
+		{
+			return bEqual;
+		}
+
+		bEqual &= FMemory::Memcmp(A.PackedGlobalArrays.GetTypedData(), B.PackedGlobalArrays.GetTypedData(), A.PackedGlobalArrays.GetTypeSize()*A.PackedGlobalArrays.Num()) == 0;
+
+		for (int32 Item = 0; Item < A.PackedUniformBuffers.Num(); Item++)
+		{
+			const TArray<FMetalPackedArrayInfo> &ArrayA = A.PackedUniformBuffers[Item];
+			const TArray<FMetalPackedArrayInfo> &ArrayB = B.PackedUniformBuffers[Item];
+
+			bEqual &= FMemory::Memcmp(ArrayA.GetTypedData(), ArrayB.GetTypedData(), ArrayA.GetTypeSize()*ArrayA.Num()) == 0;
+		}
+
+		return bEqual;
+	}
+
+	friend uint32 GetTypeHash(const FMetalShaderBindings &Binding)
+	{
+		uint32 Hash = 0;
+		Hash = Binding.InOutMask;
+		Hash |= Binding.NumSamplers << 16;
+		Hash |= Binding.NumUniformBuffers << 24;
+		Hash ^= Binding.NumUAVs;
+		Hash ^= Binding.bHasRegularUniformBuffers << 8;
+		Hash ^= FCrc::MemCrc_DEPRECATED(Binding.PackedGlobalArrays.GetTypedData(), Binding.PackedGlobalArrays.GetTypeSize()*Binding.PackedGlobalArrays.Num());
+
+		for (int32 Item = 0; Item < Binding.PackedUniformBuffers.Num(); Item++)
+		{
+			const TArray<FMetalPackedArrayInfo> &Array = Binding.PackedUniformBuffers[Item];
+			Hash ^= FCrc::MemCrc_DEPRECATED(Array.GetTypedData(), Array.GetTypeSize()* Array.Num());
+		}
+		return Hash;
+	}
+};
+
+inline FArchive& operator<<(FArchive& Ar, FMetalShaderBindings& Bindings)
+{
+	Ar << Bindings.PackedUniformBuffers;
+	Ar << Bindings.PackedGlobalArrays;
+	Ar << Bindings.InOutMask;
+	Ar << Bindings.NumSamplers;
+	Ar << Bindings.NumUniformBuffers;
+	Ar << Bindings.NumUAVs;
+	Ar << Bindings.bHasRegularUniformBuffers;
+	return Ar;
+}
+
+// Information for copying members from uniform buffers to packed
+struct FMetalUniformBufferCopyInfo
+{
+	uint16 SourceOffsetInFloats;
+	uint8 SourceUBIndex;
+	uint8 DestUBIndex;
+	uint8 DestUBTypeName;
+	uint8 DestUBTypeIndex;
+	uint16 DestOffsetInFloats;
+	uint16 SizeInFloats;
+};
+
+inline FArchive& operator<<(FArchive& Ar, FMetalUniformBufferCopyInfo& Info)
+{
+	Ar << Info.SourceOffsetInFloats;
+	Ar << Info.SourceUBIndex;
+	Ar << Info.DestUBIndex;
+	Ar << Info.DestUBTypeName;
+	if (Ar.IsLoading())
+	{
+		Info.DestUBTypeIndex = MetalPackedTypeNameToTypeIndex(Info.DestUBTypeName);
+	}
+	Ar << Info.DestOffsetInFloats;
+	Ar << Info.SizeInFloats;
+	return Ar;
+}
+
+struct FMetalCodeHeader
+{
+	uint32 Frequency;
+	FMetalShaderBindings Bindings;
+	TArray<FMetalUniformBufferCopyInfo> UniformBuffersCopyInfo;
+};
+
+inline FArchive& operator<<(FArchive& Ar, FMetalCodeHeader& Header)
+{
+	Ar << Header.Frequency;
+	Ar << Header.Bindings;
+
+	int32 NumInfos = Header.UniformBuffersCopyInfo.Num();
+	Ar << NumInfos;
+	if (Ar.IsSaving())
+	{
+		for (int32 Index = 0; Index < NumInfos; ++Index)
+		{
+			Ar << Header.UniformBuffersCopyInfo[Index];
+		}
+	}
+	else if (Ar.IsLoading())
+	{
+		Header.UniformBuffersCopyInfo.Empty(NumInfos);
+		for (int32 Index = 0; Index < NumInfos; ++Index)
+		{
+			FMetalUniformBufferCopyInfo Info;
+			Ar << Info;
+			Header.UniformBuffersCopyInfo.Add(Info);
+		}
+	}
+
+	return Ar;
+}
diff --git a/Engine/Source/Runtime/IOS/MetalRHI/Public/MetalState.h b/Engine/Source/Runtime/IOS/MetalRHI/Public/MetalState.h
new file mode 100644
--- /dev/null
+++ b/Engine/Source/Runtime/IOS/MetalRHI/Public/MetalState.h
@@ -0,0 +1,79 @@
+// Copyright 1998-2014 Epic Games, Inc. All Rights Reserved.
+
+/*=============================================================================
+	MetalState.h: Metal state definitions.
+=============================================================================*/
+
+#pragma once
+
+class FMetalSamplerState : public FRHISamplerState
+{
+public:
+	
+	/** 
+	 * Constructor/destructor
+	 */
+	FMetalSamplerState(const FSamplerStateInitializerRHI& Initializer);
+	~FMetalSamplerState();
+
+	id <MTLSamplerState> State;
+};
+
+class FMetalRasterizerState : public FRHIRasterizerState
+{
+public:
+
+	/**
+	 * Constructor/destructor
+	 */
+	FMetalRasterizerState(const FRasterizerStateInitializerRHI& Initializer);
+	~FMetalRasterizerState();
+
+	/**
+	 * Update the context with the rasterizer state
+	 */
+	void Set();
+	
+protected:
+	FRasterizerStateInitializerRHI State;
+};
+
+class FMetalDepthStencilState : public FRHIDepthStencilState
+{
+public:
+
+	/**
+	 * Constructor/destructor
+	 */
+	FMetalDepthStencilState(const FDepthStencilStateInitializerRHI& Initializer);
+	~FMetalDepthStencilState();
+	
+	/**
+	 * Set as the active depth stencil state into the context
+	 */
+	void Set();
+	
+protected:
+	id<MTLDepthStencilState> State;
+	bool bIsDepthWriteEnabled;
+	bool bIsStencilWriteEnabled;
+	
+};
+
+class FMetalBlendState : public FRHIBlendState
+{
+public:
+
+	/**
+	 * Constructor/destructor
+	 */
+	FMetalBlendState(const FBlendStateInitializerRHI& Initializer);
+	~FMetalBlendState();
+	
+	/**
+	 * Set as the active blend states
+	 */
+	void Set();
+	
+	TStaticArray<MTLBlendDescriptor*, MaxSimultaneousRenderTargets> RenderTargetStates;
+};
diff --git a/Engine/Source/Runtime/IOS/MetalRHI/Public/MetalViewport.h b/Engine/Source/Runtime/IOS/MetalRHI/Public/MetalViewport.h
new file mode 100644
--- /dev/null
+++ b/Engine/Source/Runtime/IOS/MetalRHI/Public/MetalViewport.h
@@ -0,0 +1,17 @@
+// Copyright 1998-2014 Epic Games, Inc. All Rights Reserved.
+
+/*=============================================================================
+	MetalViewport.h: Metal viewport RHI definitions.
+=============================================================================*/
+
+#pragma once
+
+class FMetalViewport : public FRHIViewport
+{
+public:
+
+	FMetalViewport(void* WindowHandle, uint32 InSizeX,uint32 InSizeY,bool bInIsFullscreen);
+	~FMetalViewport();
+
+};
+
