UE4 - Metal shader format fixes

diff --git a/Engine/Source/Developer/iOS/MetalShaderFormat/MetalShaderFormat.Build.cs b/Engine/Source/Developer/iOS/MetalShaderFormat/MetalShaderFormat.Build.cs
--- a/Engine/Source/Developer/iOS/MetalShaderFormat/MetalShaderFormat.Build.cs
+++ b/Engine/Source/Developer/iOS/MetalShaderFormat/MetalShaderFormat.Build.cs
@@ -7,6 +7,7 @@ public class MetalShaderFormat : ModuleRules
 	public MetalShaderFormat(TargetInfo Target)
 	{
 		PrivateIncludePathModuleNames.Add("TargetPlatform");
+		PublicIncludePaths.Add("Runtime/IOS/MetalRHI/Public");
 
 		PrivateDependencyModuleNames.AddRange(
 			new string[] {
diff --git a/Engine/Source/Developer/iOS/MetalShaderFormat/Private/MetalShaderCompiler.cpp b/Engine/Source/Developer/iOS/MetalShaderFormat/Private/MetalShaderCompiler.cpp
--- a/Engine/Source/Developer/iOS/MetalShaderFormat/Private/MetalShaderCompiler.cpp
+++ b/Engine/Source/Developer/iOS/MetalShaderFormat/Private/MetalShaderCompiler.cpp
@@ -3,7 +3,8 @@
 
 #include "MetalShaderFormat.h"
 #include "Core.h"
-#include "../../../EmptyRHI/Public/EmptyShaderResources.h"
+#include "ShaderCore.h"
+#include "MetalShaderResources.h"
 #include "ShaderCompilerCommon.h"
 
 #if PLATFORM_WINDOWS
diff --git a/Engine/Source/Runtime/IOS/MetalRHI/Public/MetalShaderResources.h b/Engine/Source/Runtime/IOS/MetalRHI/Public/MetalShaderResources.h
--- a/Engine/Source/Runtime/IOS/MetalRHI/Public/MetalShaderResources.h
+++ b/Engine/Source/Runtime/IOS/MetalRHI/Public/MetalShaderResources.h
@@ -73,11 +73,41 @@ inline FArchive& operator<<(FArchive& Ar, FMetalPackedArrayInfo& Info)
 	return Ar;
 }
 
+struct FMetalShaderResourceTable : public FBaseShaderResourceTable
+{
+	/** Mapping of bound Textures to their location in resource tables. */
+	TArray<uint32> TextureMap;
+	friend bool operator==(const FMetalShaderResourceTable &A, const FMetalShaderResourceTable& B)
+	{
+		if (!(((FBaseShaderResourceTable&)A) == ((FBaseShaderResourceTable&)B)))
+		{
+			return false;
+		}
+		if (A.TextureMap.Num() != B.TextureMap.Num())
+		{
+			return false;
+		}
+		if (FMemory::Memcmp(A.TextureMap.GetTypedData(), B.TextureMap.GetTypedData(), A.TextureMap.GetTypeSize()*A.TextureMap.Num()) != 0)
+		{
+			return false;
+		}
+		return true;
+	}
+};
+
+inline FArchive& operator<<(FArchive& Ar, FMetalShaderResourceTable& SRT)
+{
+	Ar << ((FBaseShaderResourceTable&)SRT);
+	Ar << SRT.TextureMap;
+	return Ar;
+}
+
 
 struct FMetalShaderBindings
 {
 	TArray<TArray<FMetalPackedArrayInfo>>	PackedUniformBuffers;
 	TArray<FMetalPackedArrayInfo>			PackedGlobalArrays;
+	FMetalShaderResourceTable				ShaderResourceTable;
 
 	uint16	InOutMask;
 	uint8	NumSamplers;
@@ -105,6 +135,7 @@ struct FMetalShaderBindings
 		bEqual &= A.bHasRegularUniformBuffers == B.bHasRegularUniformBuffers;
 		bEqual &= A.PackedGlobalArrays.Num() == B.PackedGlobalArrays.Num();
 		bEqual &= A.PackedUniformBuffers.Num() == B.PackedUniformBuffers.Num();
+		bEqual &= A.ShaderResourceTable == B.ShaderResourceTable;
 
 		if (!bEqual)
 		{
@@ -147,6 +178,7 @@ inline FArchive& operator<<(FArchive& Ar, FMetalShaderBindings& Bindings)
 {
 	Ar << Bindings.PackedUniformBuffers;
 	Ar << Bindings.PackedGlobalArrays;
+	Ar << Bindings.ShaderResourceTable;
 	Ar << Bindings.InOutMask;
 	Ar << Bindings.NumSamplers;
 	Ar << Bindings.NumUniformBuffers;
