- Added new vertex element type in Metal (Half4). Also made it crash if unknown types are added instead of failing to work right

diff --git a/Engine/Source/Runtime/IOS/MetalRHI/Private/MetalVertexDeclaration.cpp b/Engine/Source/Runtime/IOS/MetalRHI/Private/MetalVertexDeclaration.cpp
--- a/Engine/Source/Runtime/IOS/MetalRHI/Private/MetalVertexDeclaration.cpp
+++ b/Engine/Source/Runtime/IOS/MetalRHI/Private/MetalVertexDeclaration.cpp
@@ -10,19 +10,41 @@ static MTLVertexFormat TranslateElementTypeToMTLType(EVertexElementType Type)
 {
 	switch (Type)
 	{
-	case VET_Float1:		return MTLVertexFormatFloat;
-	case VET_Float2:		return MTLVertexFormatFloat2;
-	case VET_Float3:		return MTLVertexFormatFloat3;
-	case VET_Float4:		return MTLVertexFormatFloat4;
-	case VET_PackedNormal:	return MTLVertexFormatUChar4Normalized;
-	case VET_UByte4:		return MTLVertexFormatUChar4;
-	case VET_UByte4N:		return MTLVertexFormatUChar4Normalized;
-	case VET_Color:			return MTLVertexFormatUChar4Normalized;
-	case VET_Short2:		return MTLVertexFormatShort2;
-	case VET_Short4:		return MTLVertexFormatShort4;
-	case VET_Short2N:		return MTLVertexFormatShort2Normalized;
-	case VET_Half2:			return MTLVertexFormatHalf2;
-	default:				return MTLVertexFormatFloat4;
+		case VET_Float1:		return MTLVertexFormatFloat;
+		case VET_Float2:		return MTLVertexFormatFloat2;
+		case VET_Float3:		return MTLVertexFormatFloat3;
+		case VET_Float4:		return MTLVertexFormatFloat4;
+		case VET_PackedNormal:	return MTLVertexFormatUChar4Normalized;
+		case VET_UByte4:		return MTLVertexFormatUChar4;
+		case VET_UByte4N:		return MTLVertexFormatUChar4Normalized;
+		case VET_Color:			return MTLVertexFormatUChar4Normalized;
+		case VET_Short2:		return MTLVertexFormatShort2;
+		case VET_Short4:		return MTLVertexFormatShort4;
+		case VET_Short2N:		return MTLVertexFormatShort2Normalized;
+		case VET_Half2:			return MTLVertexFormatHalf2;
+		case VET_Half4:			return MTLVertexFormatHalf4;
+		default:				UE_LOG(LogMetal, Fatal, TEXT("Unknown vertex element type!"));
+	};
+}
+
+static uint32 TranslateElementTypeToSize(EVertexElementType Type)
+{
+	switch (Type)
+	{
+		case VET_Float1:		return 4;
+		case VET_Float2:		return 8;
+		case VET_Float3:		return 12;
+		case VET_Float4:		return 16;
+		case VET_PackedNormal:	return 4;
+		case VET_UByte4:		return 4;
+		case VET_UByte4N:		return 4;
+		case VET_Color:			return 4;
+		case VET_Short2:		return 4;
+		case VET_Short4:		return 8;
+		case VET_Short2N:		return 8;
+		case VET_Half2:			return 4;
+		case VET_Half4:			return 8;
+		default:				UE_LOG(LogMetal, Fatal, TEXT("Unknown vertex element type!"));
 	};
 }
 
@@ -70,6 +92,8 @@ void FMetalVertexDeclaration::GenerateLayout(const FVertexDeclarationElementList
 		// @todo urban: for zero stride elements, assume a repeated color
 		uint32 Stride = Element.Stride ? Element.Stride : 4;
 
+		checkf(Element.Offset + TranslateElementTypeToSize(Element.Type) <= Stride, TEXT("Stream component is bigger than stride: Offset: %d, Size: %d [Type %d], Stride: %d"), Element.Offset, TranslateElementTypeToSize(Element.Type), (uint32)Element.Type, Stride);
+
 		// we offset 6 buffers to leave space for uniform buffers
 		uint32 ShaderBufferIndex = UNREAL_TO_METAL_BUFFER_INDEX(Element.StreamIndex);
 
