- IOS 8 beta2 Metal updates. _Only compiles_, changes haven't been tested yet (having device update issues)

diff --git a/Engine/Source/Runtime/IOS/MetalRHI/Private/MetalManager.cpp b/Engine/Source/Runtime/IOS/MetalRHI/Private/MetalManager.cpp
--- a/Engine/Source/Runtime/IOS/MetalRHI/Private/MetalManager.cpp
+++ b/Engine/Source/Runtime/IOS/MetalRHI/Private/MetalManager.cpp
@@ -60,20 +60,20 @@ void FPipelineShadow::SetHash(uint64 InHash)
 
 	bIsDepthWriteEnabled = GET_HASH(OFFSET_DEPTH_WRITE_ENABLED, NUMBITS_DEPTH_WRITE_ENABLED);
 	bIsStencilWriteEnabled = GET_HASH(OFFSET_STENCIL_WRITE_ENABLED, NUMBITS_STENCIL_WRITE_ENABLED);
-	BlendState[0] = [[MTLBlendDescriptor alloc] init];
-	FMetalManager::Get()->ReleaseObject(BlendState[0]);
-	BlendState[0].sourceRGBBlendFactor = (MTLBlendFactor)GET_HASH(OFFSET_SOURCE_RGB_BLEND_FACTOR, NUMBITS_SOURCE_RGB_BLEND_FACTOR);
-	BlendState[0].destinationRGBBlendFactor = (MTLBlendFactor)GET_HASH(OFFSET_DEST_RGB_BLEND_FACTOR, NUMBITS_DEST_RGB_BLEND_FACTOR);
-	BlendState[0].rgbBlendOperation = (MTLBlendOperation)GET_HASH(OFFSET_RGB_BLEND_OPERATION, NUMBITS_RGB_BLEND_OPERATION);
-	BlendState[0].sourceAlphaBlendFactor = (MTLBlendFactor)GET_HASH(OFFSET_SOURCE_A_BLEND_FACTOR, NUMBITS_SOURCE_A_BLEND_FACTOR);
-	BlendState[0].destinationAlphaBlendFactor = (MTLBlendFactor)GET_HASH(OFFSET_DEST_A_BLEND_FACTOR, NUMBITS_DEST_A_BLEND_FACTOR);
-	BlendState[0].alphaBlendOperation = (MTLBlendOperation)GET_HASH(OFFSET_A_BLEND_OPERATION, NUMBITS_A_BLEND_OPERATION);
-	BlendState[0].writeMask = GET_HASH(OFFSET_WRITE_MASK, NUMBITS_WRITE_MASK);
-	BlendState[0].blendingEnabled =
-		BlendState[0].rgbBlendOperation != MTLBlendOperationAdd || BlendState[0].destinationRGBBlendFactor != MTLBlendFactorZero || BlendState[0].sourceRGBBlendFactor != MTLBlendFactorOne ||
-		BlendState[0].alphaBlendOperation != MTLBlendOperationAdd || BlendState[0].destinationAlphaBlendFactor != MTLBlendFactorZero || BlendState[0].sourceAlphaBlendFactor != MTLBlendFactorOne;
-
-	RenderTargetFormat[0] = (MTLPixelFormat)GET_HASH(OFFSET_RENDER_TARGET_FORMAT, NUMBITS_RENDER_TARGET_FORMAT);
+	RenderTargets[0] = [[MTLRenderPipelineAttachmentDescriptor alloc] init];
+	FMetalManager::Get()->ReleaseObject(RenderTargets[0]);
+	RenderTargets[0].sourceRGBBlendFactor = (MTLBlendFactor)GET_HASH(OFFSET_SOURCE_RGB_BLEND_FACTOR, NUMBITS_SOURCE_RGB_BLEND_FACTOR);
+	RenderTargets[0].destinationRGBBlendFactor = (MTLBlendFactor)GET_HASH(OFFSET_DEST_RGB_BLEND_FACTOR, NUMBITS_DEST_RGB_BLEND_FACTOR);
+	RenderTargets[0].rgbBlendOperation = (MTLBlendOperation)GET_HASH(OFFSET_RGB_BLEND_OPERATION, NUMBITS_RGB_BLEND_OPERATION);
+	RenderTargets[0].sourceAlphaBlendFactor = (MTLBlendFactor)GET_HASH(OFFSET_SOURCE_A_BLEND_FACTOR, NUMBITS_SOURCE_A_BLEND_FACTOR);
+	RenderTargets[0].destinationAlphaBlendFactor = (MTLBlendFactor)GET_HASH(OFFSET_DEST_A_BLEND_FACTOR, NUMBITS_DEST_A_BLEND_FACTOR);
+	RenderTargets[0].alphaBlendOperation = (MTLBlendOperation)GET_HASH(OFFSET_A_BLEND_OPERATION, NUMBITS_A_BLEND_OPERATION);
+	RenderTargets[0].writeMask = GET_HASH(OFFSET_WRITE_MASK, NUMBITS_WRITE_MASK);
+	RenderTargets[0].blendingEnabled =
+		RenderTargets[0].rgbBlendOperation != MTLBlendOperationAdd || RenderTargets[0].destinationRGBBlendFactor != MTLBlendFactorZero || RenderTargets[0].sourceRGBBlendFactor != MTLBlendFactorOne ||
+		RenderTargets[0].alphaBlendOperation != MTLBlendOperationAdd || RenderTargets[0].destinationAlphaBlendFactor != MTLBlendFactorZero || RenderTargets[0].sourceAlphaBlendFactor != MTLBlendFactorOne;
+
+	RenderTargets[0].pixelFormat = (MTLPixelFormat)GET_HASH(OFFSET_RENDER_TARGET_FORMAT, NUMBITS_RENDER_TARGET_FORMAT);
 	DepthTargetFormat = (MTLPixelFormat)GET_HASH(OFFSET_DEPTH_TARGET_FORMAT, NUMBITS_DEPTH_TARGET_FORMAT);
 	SampleCount = GET_HASH(OFFSET_SAMPLE_COUNT, NUMBITS_SAMPLE_COUNT);
 }
@@ -100,6 +100,46 @@ static MTLCullMode TranslateCullMode(ERasterizerCullMode CullMode)
 	}
 }
 
+static MTLBlendOperation TranslateBlendOp(EBlendOperation BlendOp)
+{
+	switch (BlendOp)
+	{
+	case BO_Subtract:	return MTLBlendOperationSubtract;
+	case BO_Min:		return MTLBlendOperationMin;
+	case BO_Max:		return MTLBlendOperationMax;
+	default:			return MTLBlendOperationAdd;
+	};
+}
+
+
+static MTLBlendFactor TranslateBlendFactor(EBlendFactor BlendFactor)
+{
+	switch (BlendFactor)
+	{
+	case BF_One:					return MTLBlendFactorOne;
+	case BF_SourceColor:			return MTLBlendFactorSourceColor;
+	case BF_InverseSourceColor:		return MTLBlendFactorOneMinusSourceColor;
+	case BF_SourceAlpha:			return MTLBlendFactorSourceAlpha;
+	case BF_InverseSourceAlpha:		return MTLBlendFactorOneMinusSourceAlpha;
+	case BF_DestAlpha:				return MTLBlendFactorDestinationAlpha;
+	case BF_InverseDestAlpha:		return MTLBlendFactorOneMinusDestinationAlpha;
+	case BF_DestColor:				return MTLBlendFactorDestinationColor;
+	case BF_InverseDestColor:		return MTLBlendFactorOneMinusDestinationColor;
+	default:						return MTLBlendFactorZero;
+	};
+}
+
+static MTLColorWriteMask TranslateWriteMask(EColorWriteMask WriteMask)
+{
+	uint32 Result = 0;
+	Result |= (WriteMask & CW_RED) ? (MTLColorWriteMaskRed) : 0;
+	Result |= (WriteMask & CW_GREEN) ? (MTLColorWriteMaskGreen) : 0;
+	Result |= (WriteMask & CW_BLUE) ? (MTLColorWriteMaskBlue) : 0;
+	Result |= (WriteMask & CW_ALPHA) ? (MTLColorWriteMaskAlpha) : 0;
+
+	return (MTLColorWriteMask)Result;
+}
+
 
 
 
@@ -116,17 +156,13 @@ id<MTLRenderPipelineState> FPipelineShadow::CreatePipelineStateForBoundShaderSta
 	// set per-MRT settings
 	for (uint32 RenderTargetIndex = 0; RenderTargetIndex < MaxMetalRenderTargets; ++RenderTargetIndex)
 	{
-		[Desc setBlendDescriptor:BlendState[RenderTargetIndex] atIndex:(MTLFramebufferAttachmentIndex)RenderTargetIndex];
-		if (RenderTargetFormat[RenderTargetIndex] != MTLPixelFormatInvalid)
-		{
-			[Desc setPixelFormat:RenderTargetFormat[RenderTargetIndex] atIndex:(MTLFramebufferAttachmentIndex)RenderTargetIndex];
-		}
+		[Desc.colorAttachments setObject:RenderTargets[RenderTargetIndex] atIndexedSubscript:RenderTargetIndex];
 	}
 
 	// depth setting if it's actually used
-	if (DepthTargetFormat != MTLPixelFormatInvalid)
+//	if (DepthTargetFormat != MTLPixelFormatInvalid)
 	{
-		[Desc setPixelFormat:DepthTargetFormat atIndex:MTLFramebufferAttachmentIndexDepth];
+		Desc.depthAttachmentPixelFormat = DepthTargetFormat;
 	}
 
 	// set the bound shader state settings
@@ -214,7 +250,8 @@ FMetalManager::FMetalManager()
 	: Device([IOSAppDelegate GetDelegate].IOSView->MetalDevice)
 	, CurrentCommandBuffer(nil)
 	, CurrentDrawable(nil)
-	, CurrentFramebuffer(nil)
+/// 	, CurrentFramebuffer(nil)
+	, CurrentRenderPass(nil)
 	, CurrentContext(nil)
 	, CurrentColorRenderTexture(nil)
 	, PreviousColorRenderTexture(nil)
@@ -373,8 +410,10 @@ void FMetalManager::EndFrame(bool bPresent)
 	}
 
 
-	ReleaseObject(CurrentFramebuffer);
-	CurrentFramebuffer = nil;
+/// 	ReleaseObject(CurrentFramebuffer);
+/// 	CurrentFramebuffer = nil;
+	ReleaseObject(CurrentRenderPass);
+	CurrentRenderPass = nil;
 
 
 	UNTRACK_OBJECT(CurrentCommandBuffer);
@@ -498,17 +537,37 @@ void FMetalManager::SetBlendState(FMetalBlendState* BlendState)
 {
 	for(uint32 RenderTargetIndex = 0;RenderTargetIndex < MaxMetalRenderTargets; ++RenderTargetIndex)
 	{
-		MTLBlendDescriptor* Blend0 = BlendState->RenderTargetStates[RenderTargetIndex];
+		// @todo metal mrt
 		check(RenderTargetIndex == 0);
 
-		Pipeline.BlendState[RenderTargetIndex] = Blend0;
-		SET_HASH(OFFSET_SOURCE_RGB_BLEND_FACTOR, NUMBITS_SOURCE_RGB_BLEND_FACTOR, Blend0.sourceRGBBlendFactor);
-		SET_HASH(OFFSET_DEST_RGB_BLEND_FACTOR, NUMBITS_DEST_RGB_BLEND_FACTOR, Blend0.destinationRGBBlendFactor);
-		SET_HASH(OFFSET_RGB_BLEND_OPERATION, NUMBITS_RGB_BLEND_OPERATION, Blend0.rgbBlendOperation);
-		SET_HASH(OFFSET_SOURCE_A_BLEND_FACTOR, NUMBITS_SOURCE_A_BLEND_FACTOR, Blend0.sourceAlphaBlendFactor);
-		SET_HASH(OFFSET_DEST_A_BLEND_FACTOR, NUMBITS_DEST_A_BLEND_FACTOR, Blend0.destinationAlphaBlendFactor);
-		SET_HASH(OFFSET_A_BLEND_OPERATION, NUMBITS_A_BLEND_OPERATION, Blend0.alphaBlendOperation);
-		SET_HASH(OFFSET_WRITE_MASK, NUMBITS_WRITE_MASK, Blend0.writeMask);
+		// get the blend params
+		const FBlendStateInitializerRHI::FRenderTarget& Init =
+			BlendState->RenderTargetStates.bUseIndependentRenderTargetBlendStates
+				? BlendState->RenderTargetStates.RenderTargets[RenderTargetIndex]
+				: BlendState->RenderTargetStates.RenderTargets[0];
+
+		// write into the pipline attachment
+		MTLRenderPipelineAttachmentDescriptor* Attachment = Pipeline.RenderTargets[RenderTargetIndex];
+		
+		// set values
+		Attachment.blendingEnabled =
+			Init.ColorBlendOp != BO_Add || Init.ColorDestBlend != BF_Zero || Init.ColorSrcBlend != BF_One ||
+			Init.AlphaBlendOp != BO_Add || Init.AlphaDestBlend != BF_Zero || Init.AlphaSrcBlend != BF_One;
+		Attachment.sourceRGBBlendFactor = TranslateBlendFactor(Init.ColorSrcBlend);
+		Attachment.destinationRGBBlendFactor = TranslateBlendFactor(Init.ColorDestBlend);
+		Attachment.rgbBlendOperation = TranslateBlendOp(Init.ColorBlendOp);
+		Attachment.sourceAlphaBlendFactor = TranslateBlendFactor(Init.AlphaSrcBlend);
+		Attachment.destinationAlphaBlendFactor = TranslateBlendFactor(Init.AlphaDestBlend);
+		Attachment.alphaBlendOperation = TranslateBlendOp(Init.AlphaBlendOp);
+		Attachment.writeMask = TranslateWriteMask(Init.ColorWriteMask);
+
+		SET_HASH(OFFSET_SOURCE_RGB_BLEND_FACTOR, NUMBITS_SOURCE_RGB_BLEND_FACTOR, Attachment.sourceRGBBlendFactor);
+		SET_HASH(OFFSET_DEST_RGB_BLEND_FACTOR, NUMBITS_DEST_RGB_BLEND_FACTOR, Attachment.destinationRGBBlendFactor);
+		SET_HASH(OFFSET_RGB_BLEND_OPERATION, NUMBITS_RGB_BLEND_OPERATION, Attachment.rgbBlendOperation);
+		SET_HASH(OFFSET_SOURCE_A_BLEND_FACTOR, NUMBITS_SOURCE_A_BLEND_FACTOR, Attachment.sourceAlphaBlendFactor);
+		SET_HASH(OFFSET_DEST_A_BLEND_FACTOR, NUMBITS_DEST_A_BLEND_FACTOR, Attachment.destinationAlphaBlendFactor);
+		SET_HASH(OFFSET_A_BLEND_OPERATION, NUMBITS_A_BLEND_OPERATION, Attachment.alphaBlendOperation);
+		SET_HASH(OFFSET_WRITE_MASK, NUMBITS_WRITE_MASK, Attachment.writeMask);
 	}
 }
 
@@ -582,17 +641,35 @@ void FMetalManager::UpdateContext()
 		return;
 	}
 
-	MTLAttachmentDescriptor* ColorAttachment = nil;
-	MTLAttachmentDescriptor* DepthAttachment = nil;
+
+	// toss previous render pass object
+	if (CurrentRenderPass != nil)
+	{
+		ReleaseObject(CurrentRenderPass);
+	}
+
+	// make a new one
+	CurrentRenderPass = [MTLRenderPassDescriptor renderPassDescriptor];
+	TRACK_OBJECT(CurrentRenderPass);
+	// if we need to do queries, write to the ring buffer (we set the offset into the ring buffer per query)
+	CurrentRenderPass.visibilityResultBuffer = QueryBuffer.Buffer;
+
+	// @todo metal mrt: Set an attachment index properly
+	uint32 AttachmentIndex = 0;
+
+	MTLRenderPassAttachmentDescriptor* ColorAttachment = nil;
+	MTLRenderPassAttachmentDescriptor* DepthAttachment = nil;
 	// @todo urban: Handle stencil
-	MTLAttachmentDescriptor* StencilAttachment = nil;
+	MTLRenderPassAttachmentDescriptor* StencilAttachment = nil;
 
 	Pipeline.SampleCount = 0;
 	if (CurrentColorRenderTexture)
 	{
+		ColorAttachment = [[MTLRenderPassAttachmentDescriptor alloc] init];
+
 		if (CurrentMSAARenderTexture)
 		{
-			ColorAttachment = [MTLAttachmentDescriptor attachmentDescriptorWithTexture:CurrentMSAARenderTexture];
+			ColorAttachment.texture = CurrentMSAARenderTexture;
 			[ColorAttachment setLoadAction:MTLLoadActionDontCare];
 			[ColorAttachment setStoreAction:MTLStoreActionMultisampleResolve];
 			[ColorAttachment setResolveTexture:CurrentColorRenderTexture];
@@ -600,23 +677,28 @@ void FMetalManager::UpdateContext()
 		}
 		else
 		{
-			ColorAttachment = [MTLAttachmentDescriptor attachmentDescriptorWithTexture:CurrentColorRenderTexture];
+			ColorAttachment.texture = CurrentColorRenderTexture;
 			[ColorAttachment setLoadAction:MTLLoadActionDontCare];
 			[ColorAttachment setStoreAction:MTLStoreActionStore];
 			Pipeline.SampleCount = 1;
 		}
 
-		Pipeline.RenderTargetFormat[0] = CurrentColorRenderTexture.pixelFormat;
+		[CurrentRenderPass.colorAttachments setObject:ColorAttachment atIndexedSubscript:AttachmentIndex];
+		[ColorAttachment release];
+
+		Pipeline.RenderTargets[AttachmentIndex].pixelFormat = CurrentColorRenderTexture.pixelFormat;
 	}
 	else
 	{
-		Pipeline.RenderTargetFormat[0] = MTLPixelFormatInvalid;
+		Pipeline.RenderTargets[AttachmentIndex].pixelFormat = MTLPixelFormatInvalid;
 	}
-	SET_HASH(OFFSET_RENDER_TARGET_FORMAT, NUMBITS_RENDER_TARGET_FORMAT, Pipeline.RenderTargetFormat[0]);
+	SET_HASH(OFFSET_RENDER_TARGET_FORMAT, NUMBITS_RENDER_TARGET_FORMAT, Pipeline.RenderTargets[AttachmentIndex].pixelFormat);
 	
 	if (CurrentDepthRenderTexture)
 	{
-		DepthAttachment = [MTLAttachmentDescriptor attachmentDescriptorWithTexture : CurrentDepthRenderTexture];
+		DepthAttachment = [[MTLRenderPassAttachmentDescriptor alloc] init];
+
+		DepthAttachment.texture = CurrentDepthRenderTexture;
 		[DepthAttachment setLoadAction:MTLLoadActionClear];
 		[DepthAttachment setStoreAction:MTLStoreActionDontCare];
 		[DepthAttachment setClearValue:MTLClearValueMakeDepth(0.0f)];
@@ -626,6 +708,9 @@ void FMetalManager::UpdateContext()
 		{
 			Pipeline.SampleCount = CurrentDepthRenderTexture.sampleCount;
 		}
+
+		CurrentRenderPass.depthAttachment = DepthAttachment;
+		[DepthAttachment release];
 	}
 	else
 	{
@@ -652,7 +737,7 @@ void FMetalManager::UpdateContext()
 		}
 
 		ReleaseObject(CurrentContext);
-		ReleaseObject(CurrentFramebuffer);
+///		ReleaseObject(CurrentFramebuffer);
 
 		// if we are doing occlusion queries, we could use this method, along with a completion callback
 		// to set a "render target complete" flag that the OQ code next frame would wait on
@@ -683,16 +768,15 @@ void FMetalManager::UpdateContext()
 	}
 		
 	// make a new framebuffer object to render to
-	MTLFramebufferDescriptor* FramebufferDescriptor = [MTLFramebufferDescriptor framebufferDescriptorWithColorAttachment:ColorAttachment depthAttachment:DepthAttachment stencilAttachment:nil];
+///	MTLFramebufferDescriptor* FramebufferDescriptor = [MTLFramebufferDescriptor framebufferDescriptorWithColorAttachment:ColorAttachment depthAttachment:DepthAttachment stencilAttachment:nil];
 	// if we need to do queries, write to the ring buffer (we set the offset into the ring buffer per query)
-	FramebufferDescriptor.visibilityResultBuffer = QueryBuffer.Buffer;
+///	FramebufferDescriptor.visibilityResultBuffer = QueryBuffer.Buffer;
 
-	CurrentFramebuffer = [Device newFramebufferWithDescriptor:FramebufferDescriptor];
-	TRACK_OBJECT(CurrentFramebuffer);
+///	CurrentFramebuffer = [Device newFramebufferWithDescriptor:FramebufferDescriptor];
+///	TRACK_OBJECT(CurrentFramebuffer);
 
 	// make a new render context to use to render to the framebuffer
-	CurrentContext = [CurrentCommandBuffer renderCommandEncoderWithFramebuffer:CurrentFramebuffer];
-//	[CurrentContext setLabel : @"HelloWorld"];
+	CurrentContext = [CurrentCommandBuffer renderCommandEncoderWithDescriptor:CurrentRenderPass];
 	[CurrentContext retain];
 	TRACK_OBJECT(CurrentContext);
 
diff --git a/Engine/Source/Runtime/IOS/MetalRHI/Private/MetalManager.h b/Engine/Source/Runtime/IOS/MetalRHI/Private/MetalManager.h
--- a/Engine/Source/Runtime/IOS/MetalRHI/Private/MetalManager.h
+++ b/Engine/Source/Runtime/IOS/MetalRHI/Private/MetalManager.h
@@ -14,6 +14,11 @@ struct FPipelineShadow
 	{
 		// zero out our memory
 		FMemory::Memzero(this, sizeof(*this));
+
+		for (int Index = 0; Index < MaxMetalRenderTargets; Index++)
+		{
+			RenderTargets[Index] = [[MTLRenderPipelineAttachmentDescriptor alloc] init];
+		}
 	}
 
 	/**
@@ -36,8 +41,7 @@ struct FPipelineShadow
 
 	bool bIsDepthWriteEnabled;
 	bool bIsStencilWriteEnabled;
-	MTLBlendDescriptor* BlendState[MaxMetalRenderTargets];
-	MTLPixelFormat RenderTargetFormat[MaxMetalRenderTargets];
+	MTLRenderPipelineAttachmentDescriptor* RenderTargets[MaxMetalRenderTargets];
 	MTLPixelFormat DepthTargetFormat;
     uint32 SampleCount;
 
@@ -145,7 +149,7 @@ protected:
 
 	dispatch_semaphore_t CommandBufferSemaphore;
 
-	id<MTLFramebuffer> CurrentFramebuffer;
+	MTLRenderPassDescriptor* CurrentRenderPass;
 	id<MTLRenderCommandEncoder> CurrentContext;
 	
 	id<MTLTexture> CurrentColorRenderTexture, PreviousColorRenderTexture;
diff --git a/Engine/Source/Runtime/IOS/MetalRHI/Private/MetalState.cpp b/Engine/Source/Runtime/IOS/MetalRHI/Private/MetalState.cpp
--- a/Engine/Source/Runtime/IOS/MetalRHI/Private/MetalState.cpp
+++ b/Engine/Source/Runtime/IOS/MetalRHI/Private/MetalState.cpp
@@ -100,45 +100,6 @@ static MTLStencilOperation TranslateStencilOp(EStencilOp StencilOp)
 	};
 }
 
-static MTLBlendOperation TranslateBlendOp(EBlendOperation BlendOp)
-{
-	switch(BlendOp)
-	{
-		case BO_Subtract:	return MTLBlendOperationSubtract;
-		case BO_Min:		return MTLBlendOperationMin;
-		case BO_Max:		return MTLBlendOperationMax;
-		default:			return MTLBlendOperationAdd;
-	};
-}
-
-
-static MTLBlendFactor TranslateBlendFactor(EBlendFactor BlendFactor)
-{
-	switch(BlendFactor)
-	{
-		case BF_One:					return MTLBlendFactorOne;
-		case BF_SourceColor:			return MTLBlendFactorSourceColor;
-		case BF_InverseSourceColor:		return MTLBlendFactorOneMinusSourceColor;
-		case BF_SourceAlpha:			return MTLBlendFactorSourceAlpha;
-		case BF_InverseSourceAlpha:		return MTLBlendFactorOneMinusSourceAlpha;
-		case BF_DestAlpha:				return MTLBlendFactorDestinationAlpha;
-		case BF_InverseDestAlpha:		return MTLBlendFactorOneMinusDestinationAlpha;
-		case BF_DestColor:				return MTLBlendFactorDestinationColor;
-		case BF_InverseDestColor:		return MTLBlendFactorOneMinusDestinationColor;
-		default:						return MTLBlendFactorZero;
-	};
-}
-
-static MTLColorWriteMask TranslateWriteMask(EColorWriteMask WriteMask)
-{
-	uint32 Result = 0;
-	Result |= (WriteMask & CW_RED) ? (MTLColorWriteMaskRed) : 0;
-	Result |= (WriteMask & CW_GREEN) ? (MTLColorWriteMaskGreen) : 0;
-	Result |= (WriteMask & CW_BLUE) ? (MTLColorWriteMaskBlue) : 0;
-	Result |= (WriteMask & CW_ALPHA) ? (MTLColorWriteMaskAlpha) : 0;
-	
-	return (MTLColorWriteMask)Result;
-}
 
 
 
@@ -199,7 +160,7 @@ FMetalDepthStencilState::FMetalDepthStencilState(const FDepthStencilStateInitial
 		Stencil.depthStencilPassOperation = TranslateStencilOp(Initializer.FrontFacePassStencilOp);
 		Stencil.readMask = Initializer.StencilReadMask;
 		Stencil.writeMask = Initializer.StencilWriteMask;
-		Desc.frontFaceStencilDescriptor = Stencil;
+		Desc.frontFaceStencil = Stencil;
 		
 		[Stencil release];
 	}
@@ -214,7 +175,7 @@ FMetalDepthStencilState::FMetalDepthStencilState(const FDepthStencilStateInitial
 		Stencil.depthStencilPassOperation = TranslateStencilOp(Initializer.BackFacePassStencilOp);
 		Stencil.readMask = Initializer.StencilReadMask;
 		Stencil.writeMask = Initializer.StencilWriteMask;
-		Desc.backFaceStencilDescriptor = Stencil;
+		Desc.backFaceStencil = Stencil;
 		
 		[Stencil release];
 	}
@@ -247,42 +208,17 @@ void FMetalDepthStencilState::Set()
 }
 
 FMetalBlendState::FMetalBlendState(const FBlendStateInitializerRHI& Initializer)
+	: RenderTargetStates(Initializer)
 {
-	for(uint32 RenderTargetIndex = 0;RenderTargetIndex < MaxMetalRenderTargets; ++RenderTargetIndex)
-	{
-		// which initializer to use
-		const FBlendStateInitializerRHI::FRenderTarget& Init =
-			Initializer.bUseIndependentRenderTargetBlendStates
-				? Initializer.RenderTargets[RenderTargetIndex]
-				: Initializer.RenderTargets[0];
-
-		// make a new blend state
-		MTLBlendDescriptor* BlendState = [[MTLBlendDescriptor alloc] init];
-		TRACK_OBJECT(BlendState);
-
-		// set values
-		BlendState.blendingEnabled =
-			Init.ColorBlendOp != BO_Add || Init.ColorDestBlend != BF_Zero || Init.ColorSrcBlend != BF_One ||
-			Init.AlphaBlendOp != BO_Add || Init.AlphaDestBlend != BF_Zero || Init.AlphaSrcBlend != BF_One;
-		BlendState.sourceRGBBlendFactor = TranslateBlendFactor(Init.ColorSrcBlend);
-		BlendState.destinationRGBBlendFactor = TranslateBlendFactor(Init.ColorDestBlend);
-		BlendState.rgbBlendOperation = TranslateBlendOp(Init.ColorBlendOp);
-		BlendState.sourceAlphaBlendFactor = TranslateBlendFactor(Init.AlphaSrcBlend);
-		BlendState.destinationAlphaBlendFactor = TranslateBlendFactor(Init.AlphaDestBlend);
-		BlendState.alphaBlendOperation = TranslateBlendOp(Init.AlphaBlendOp);
-		BlendState.writeMask = TranslateWriteMask(Init.ColorWriteMask);
-		
-		RenderTargetStates[RenderTargetIndex] = BlendState;
-	}
 }
 
 FMetalBlendState::~FMetalBlendState()
 {
-	for(uint32 RenderTargetIndex = 0;RenderTargetIndex < MaxMetalRenderTargets; ++RenderTargetIndex)
-	{
-		UNTRACK_OBJECT(RenderTargetStates[RenderTargetIndex]);
-		[RenderTargetStates[RenderTargetIndex] release];
-	}
+// 	for(uint32 RenderTargetIndex = 0;RenderTargetIndex < MaxMetalRenderTargets; ++RenderTargetIndex)
+// 	{
+// 		UNTRACK_OBJECT(RenderTargetStates[RenderTargetIndex]);
+// 		[RenderTargetStates[RenderTargetIndex] release];
+// 	}
 }
 
 void FMetalBlendState::Set()
diff --git a/Engine/Source/Runtime/IOS/MetalRHI/Private/MetalTexture.cpp b/Engine/Source/Runtime/IOS/MetalRHI/Private/MetalTexture.cpp
--- a/Engine/Source/Runtime/IOS/MetalRHI/Private/MetalTexture.cpp
+++ b/Engine/Source/Runtime/IOS/MetalRHI/Private/MetalTexture.cpp
@@ -232,9 +232,11 @@ void FMetalSurface::Unlock(uint32 MipIndex, uint32 ArrayIndex)
 		Stride = 0;
 	}
 
+	checkf(SizeZ <= 1, TEXT("3D textures are not supported yet"));
+
 	// upload the texture to the texture slice
-	// @todo urban: The MipSizeZ would be needed for a volume texture - do we care?
-	[Texture copyFromPixels:LockedMemory rowBytes:Stride imageBytes:MipBytes toSlice:ArrayIndex mipmapLevel:MipIndex origin:MTLOriginMake(0, 0, 0) size:MTLSizeMake(FMath::Max<uint32>(SizeX>>MipIndex, 1), FMath::Max<uint32>(SizeY>>MipIndex, 1), 1)];
+	MTLTextureRegion Region = MTLTextureRegionMake2D(0, 0, FMath::Max<uint32>(SizeX>>MipIndex, 1), FMath::Max<uint32>(SizeY>>MipIndex, 1));
+	[Texture replaceRegion:Region mipmapLevel:MipIndex withBytes:LockedMemory bytesPerRow:Stride];
 	
 	FMemory::Free(LockedMemory);
 	LockedMemory = NULL;
diff --git a/Engine/Source/Runtime/IOS/MetalRHI/Private/MetalVertexDeclaration.cpp b/Engine/Source/Runtime/IOS/MetalRHI/Private/MetalVertexDeclaration.cpp
--- a/Engine/Source/Runtime/IOS/MetalRHI/Private/MetalVertexDeclaration.cpp
+++ b/Engine/Source/Runtime/IOS/MetalRHI/Private/MetalVertexDeclaration.cpp
@@ -23,8 +23,9 @@ static MTLVertexFormat TranslateElementTypeToMTLType(EVertexElementType Type)
 		case VET_Short2N:		return MTLVertexFormatShort2Normalized;
 		case VET_Half2:			return MTLVertexFormatHalf2;
 		case VET_Half4:			return MTLVertexFormatHalf4;
-		default:				UE_LOG(LogMetal, Fatal, TEXT("Unknown vertex element type!"));
+		default:				UE_LOG(LogMetal, Fatal, TEXT("Unknown vertex element type!")); return MTLVertexFormatFloat;
 	};
+
 }
 
 static uint32 TranslateElementTypeToSize(EVertexElementType Type)
@@ -44,7 +45,7 @@ static uint32 TranslateElementTypeToSize(EVertexElementType Type)
 		case VET_Short2N:		return 8;
 		case VET_Half2:			return 4;
 		case VET_Half4:			return 8;
-		default:				UE_LOG(LogMetal, Fatal, TEXT("Unknown vertex element type!"));
+		default:				UE_LOG(LogMetal, Fatal, TEXT("Unknown vertex element type!")); return 0;
 	};
 }
 
@@ -90,7 +91,8 @@ void FMetalVertexDeclaration::GenerateLayout(const FVertexDeclarationElementList
 		const FVertexElement& Element = Elements[ElementIndex];
 		
 		// @todo urban: for zero stride elements, assume a repeated color
-		uint32 Stride = Element.Stride ? Element.Stride : 4;
+		// @todo urban: Make sure zero stride actually works now!! If so, then we can remove the hack zero stride buffer
+		uint32 Stride = Element.Stride;// ? Element.Stride : 4;
 
 		checkf(Element.Offset + TranslateElementTypeToSize(Element.Type) <= Stride, TEXT("Stream component is bigger than stride: Offset: %d, Size: %d [Type %d], Stride: %d"), Element.Offset, TranslateElementTypeToSize(Element.Type), (uint32)Element.Type, Stride);
 
@@ -119,6 +121,6 @@ void FMetalVertexDeclaration::GenerateLayout(const FVertexDeclarationElementList
 			NSLog(@"Setting illegal stride - break here if you want to find out why, but this won't break until we try to render with it");
 			Stride = 200;// 16;
 		}
-		[Layout setStride:Stride instanceStepRate:0 atVertexBufferIndex:It.Key()];
+		[Layout setStride:Stride stepFunction:(Stride == 0 ? MTLVertexStepFunctionConstant : MTLVertexStepFunctionPerVertex) stepRate:0 atVertexBufferIndex:It.Key()];
 	}
 }
diff --git a/Engine/Source/Runtime/IOS/MetalRHI/Public/MetalState.h b/Engine/Source/Runtime/IOS/MetalRHI/Public/MetalState.h
--- a/Engine/Source/Runtime/IOS/MetalRHI/Public/MetalState.h
+++ b/Engine/Source/Runtime/IOS/MetalRHI/Public/MetalState.h
@@ -75,5 +75,5 @@ public:
 	 */
 	void Set();
 	
-	TStaticArray<MTLBlendDescriptor*, MaxSimultaneousRenderTargets> RenderTargetStates;
+	FBlendStateInitializerRHI RenderTargetStates;
 };
