- Bug fixes for IOS 8.0b2 API changes

diff --git a/Engine/Source/Runtime/IOS/MetalRHI/Private/MetalCommands.cpp b/Engine/Source/Runtime/IOS/MetalRHI/Private/MetalCommands.cpp
--- a/Engine/Source/Runtime/IOS/MetalRHI/Private/MetalCommands.cpp
+++ b/Engine/Source/Runtime/IOS/MetalRHI/Private/MetalCommands.cpp
@@ -85,6 +85,11 @@ void FMetalDynamicRHI::RHISetScissorRect(bool bEnable,uint32 MinX,uint32 MinY,ui
 	Scissor.width = MaxX - MinX;
 	Scissor.height = MaxY - MinY;
 
+	// metal doesn't support 0 sized scissor rect
+	if (Scissor.width == 0 || Scissor.height == 0)
+	{
+		return;
+	}
 	[FMetalManager::GetContext() setScissorRect:Scissor];
 }
 
diff --git a/Engine/Source/Runtime/IOS/MetalRHI/Private/MetalTexture.cpp b/Engine/Source/Runtime/IOS/MetalRHI/Private/MetalTexture.cpp
--- a/Engine/Source/Runtime/IOS/MetalRHI/Private/MetalTexture.cpp
+++ b/Engine/Source/Runtime/IOS/MetalRHI/Private/MetalTexture.cpp
@@ -232,11 +232,11 @@ void FMetalSurface::Unlock(uint32 MipIndex, uint32 ArrayIndex)
 		Stride = 0;
 	}
 
-	checkf(SizeZ <= 1, TEXT("3D textures are not supported yet"));
+	checkf(SizeZ <= 1 || SizeZ == 6, TEXT("3D textures are not supported yet (SizeZ = %d"), SizeZ);
 
 	// upload the texture to the texture slice
 	MTLTextureRegion Region = MTLTextureRegionMake2D(0, 0, FMath::Max<uint32>(SizeX>>MipIndex, 1), FMath::Max<uint32>(SizeY>>MipIndex, 1));
-	[Texture replaceRegion:Region mipmapLevel:MipIndex withBytes:LockedMemory bytesPerRow:Stride];
+	[Texture replaceRegion:Region mipmapLevel:MipIndex slice:ArrayIndex withBytes:LockedMemory bytesPerRow:Stride bytesPerImage:MipBytes];
 	
 	FMemory::Free(LockedMemory);
 	LockedMemory = NULL;
diff --git a/Engine/Source/Runtime/IOS/MetalRHI/Private/MetalVertexDeclaration.cpp b/Engine/Source/Runtime/IOS/MetalRHI/Private/MetalVertexDeclaration.cpp
--- a/Engine/Source/Runtime/IOS/MetalRHI/Private/MetalVertexDeclaration.cpp
+++ b/Engine/Source/Runtime/IOS/MetalRHI/Private/MetalVertexDeclaration.cpp
@@ -94,7 +94,8 @@ void FMetalVertexDeclaration::GenerateLayout(const FVertexDeclarationElementList
 		// @todo urban: Make sure zero stride actually works now!! If so, then we can remove the hack zero stride buffer
 		uint32 Stride = Element.Stride;// ? Element.Stride : 4;
 
-		checkf(Element.Offset + TranslateElementTypeToSize(Element.Type) <= Stride, TEXT("Stream component is bigger than stride: Offset: %d, Size: %d [Type %d], Stride: %d"), Element.Offset, TranslateElementTypeToSize(Element.Type), (uint32)Element.Type, Stride);
+		checkf(Stride == 0 || Element.Offset + TranslateElementTypeToSize(Element.Type) <= Stride, 
+			TEXT("Stream component is bigger than stride: Offset: %d, Size: %d [Type %d], Stride: %d"), Element.Offset, TranslateElementTypeToSize(Element.Type), (uint32)Element.Type, Stride);
 
 		// we offset 6 buffers to leave space for uniform buffers
 		uint32 ShaderBufferIndex = UNREAL_TO_METAL_BUFFER_INDEX(Element.StreamIndex);
@@ -121,6 +122,11 @@ void FMetalVertexDeclaration::GenerateLayout(const FVertexDeclarationElementList
 			NSLog(@"Setting illegal stride - break here if you want to find out why, but this won't break until we try to render with it");
 			Stride = 200;// 16;
 		}
-		[Layout setStride:Stride stepFunction:(Stride == 0 ? MTLVertexStepFunctionConstant : MTLVertexStepFunctionPerVertex) stepRate:0 atVertexBufferIndex:It.Key()];
+
+		// handle 0 stride buffers
+		MTLVertexStepFunction Function = (Stride == 0 ? MTLVertexStepFunctionConstant : MTLVertexStepFunctionPerVertex);
+		uint32 StepRate = (Stride == 0 ? 0 : 1);
+		uint32 UsedStride = (Stride == 0 ? 4 : Stride);
+		[Layout setStride:UsedStride stepFunction:Function stepRate:StepRate atVertexBufferIndex:It.Key()];
 	}
 }
