- More SRT fixes for Metal
- Added RHISetRenderTarget call after the Slate RHIBeginDrawingViewport code so that platform's that don't set the back buffer in BeginDrawingViewport won't fail

diff --git a/Engine/Source/Runtime/IOS/MetalRHI/Private/MetalCommands.cpp b/Engine/Source/Runtime/IOS/MetalRHI/Private/MetalCommands.cpp
--- a/Engine/Source/Runtime/IOS/MetalRHI/Private/MetalCommands.cpp
+++ b/Engine/Source/Runtime/IOS/MetalRHI/Private/MetalCommands.cpp
@@ -223,7 +223,7 @@ void FMetalDynamicRHI::RHISetShaderSampler(FComputeShaderRHIParamRef ComputeShad
 
 void FMetalDynamicRHI::RHISetShaderParameter(FVertexShaderRHIParamRef VertexShaderRHI, uint32 BufferIndex, uint32 BaseIndex, uint32 NumBytes, const void* NewValue)
 {
-	FMetalManager::Get()->ShaderParameters[METAL_SHADER_STAGE_VERTEX].Set(BufferIndex, BaseIndex, NumBytes, NewValue);
+	FMetalManager::Get()->GetShaderParameters(METAL_SHADER_STAGE_VERTEX).Set(BufferIndex, BaseIndex, NumBytes, NewValue);
 }
 
 void FMetalDynamicRHI::RHISetShaderParameter(FHullShaderRHIParamRef HullShaderRHI, uint32 BufferIndex, uint32 BaseIndex, uint32 NumBytes, const void* NewValue)
@@ -234,7 +234,7 @@ void FMetalDynamicRHI::RHISetShaderParameter(FHullShaderRHIParamRef HullShaderRH
 
 void FMetalDynamicRHI::RHISetShaderParameter(FPixelShaderRHIParamRef PixelShaderRHI, uint32 BufferIndex, uint32 BaseIndex, uint32 NumBytes, const void* NewValue)
 {
-	FMetalManager::Get()->ShaderParameters[METAL_SHADER_STAGE_PIXEL].Set(BufferIndex, BaseIndex, NumBytes, NewValue);
+	FMetalManager::Get()->GetShaderParameters(METAL_SHADER_STAGE_PIXEL).Set(BufferIndex, BaseIndex, NumBytes, NewValue);
 }
 
 void FMetalDynamicRHI::RHISetShaderParameter(FDomainShaderRHIParamRef DomainShaderRHI, uint32 BufferIndex, uint32 BaseIndex, uint32 NumBytes, const void* NewValue)
@@ -258,6 +258,10 @@ void FMetalDynamicRHI::RHISetShaderUniformBuffer(FVertexShaderRHIParamRef Vertex
 {
 	DYNAMIC_CAST_METGALRESOURCE(VertexShader, VertexShader);
 	VertexShader->BoundUniformBuffers[BufferIndex] = BufferRHI;
+	VertexShader->DirtyUniformBuffers |= 1 << BufferIndex;
+
+	// @todo metal log: remove
+	NSLog(@"Setting Vertex UB at index %d", BufferIndex);
 }
 
 void FMetalDynamicRHI::RHISetShaderUniformBuffer(FHullShaderRHIParamRef HullShader, uint32 BufferIndex, FUniformBufferRHIParamRef BufferRHI)
@@ -279,6 +283,10 @@ void FMetalDynamicRHI::RHISetShaderUniformBuffer(FPixelShaderRHIParamRef PixelSh
 {
 	DYNAMIC_CAST_METGALRESOURCE(PixelShader, PixelShader);
 	PixelShader->BoundUniformBuffers[BufferIndex] = BufferRHI;
+	PixelShader->DirtyUniformBuffers |= 1 << BufferIndex;
+
+	// @todo metal log: remove
+	NSLog(@"Setting Pixel UB at index %d", BufferIndex);
 }
 
 void FMetalDynamicRHI::RHISetShaderUniformBuffer(FComputeShaderRHIParamRef ComputeShader, uint32 BufferIndex, FUniformBufferRHIParamRef BufferRHI)
diff --git a/Engine/Source/Runtime/IOS/MetalRHI/Private/MetalManager.cpp b/Engine/Source/Runtime/IOS/MetalRHI/Private/MetalManager.cpp
--- a/Engine/Source/Runtime/IOS/MetalRHI/Private/MetalManager.cpp
+++ b/Engine/Source/Runtime/IOS/MetalRHI/Private/MetalManager.cpp
@@ -250,7 +250,6 @@ FMetalManager::FMetalManager()
 	: Device([IOSAppDelegate GetDelegate].IOSView->MetalDevice)
 	, CurrentCommandBuffer(nil)
 	, CurrentDrawable(nil)
-/// 	, CurrentFramebuffer(nil)
 	, CurrentRenderPass(nil)
 	, CurrentContext(nil)
 	, CurrentColorRenderTexture(nil)
@@ -262,7 +261,8 @@ FMetalManager::FMetalManager()
 	, WhichFreeList(0)
 	, CommandBufferIndex(0)
 	, CompletedCommandBufferIndex(0)
-
+	, SceneFrameCounter(0)
+	, ResourceTableFrameCounter(INDEX_NONE)
 {
 	CommandQueue = [Device newCommandQueue];
 
@@ -296,6 +296,28 @@ FMetalManager::FMetalManager()
 	InitFrame();
 }
 
+void FMetalManager::BeginScene()
+{
+	// Increment the frame counter. INDEX_NONE is a special value meaning "uninitialized", so if
+	// we hit it just wrap around to zero.
+	SceneFrameCounter++;
+	if (SceneFrameCounter == INDEX_NONE)
+	{
+		SceneFrameCounter++;
+	}
+
+	static auto* ResourceTableCachingCvar = IConsoleManager::Get().FindTConsoleVariableDataInt(TEXT("rhi.ResourceTableCaching"));
+	if (ResourceTableCachingCvar == NULL || ResourceTableCachingCvar->GetValueOnAnyThread() == 1)
+	{
+		ResourceTableFrameCounter = SceneFrameCounter;
+	}
+}
+
+void FMetalManager::EndScene()
+{
+	ResourceTableFrameCounter = INDEX_NONE;
+}
+
 void FMetalManager::BeginFrame()
 {
 }
@@ -410,8 +432,6 @@ void FMetalManager::EndFrame(bool bPresent)
 	}
 
 
-/// 	ReleaseObject(CurrentFramebuffer);
-/// 	CurrentFramebuffer = nil;
 	ReleaseObject(CurrentRenderPass);
 	CurrentRenderPass = nil;
 
@@ -488,6 +508,7 @@ void FMetalManager::PrepareToDraw(uint32 NumVertices)
 	// make sure the BSS has a valid pipeline state object
 	CurrentBoundShaderState->PrepareToDraw(Pipeline);
 	
+	CommitGraphicsResourceTables();
 	CommitNonComputeShaderConstants();
 
 	// grow buffers as needed
@@ -641,15 +662,12 @@ void FMetalManager::UpdateContext()
 		return;
 	}
 
-
 	// toss previous render pass object
-	if (CurrentRenderPass != nil)
-	{
-		ReleaseObject(CurrentRenderPass);
-	}
+	ReleaseObject(CurrentRenderPass);
 
 	// make a new one
 	CurrentRenderPass = [MTLRenderPassDescriptor renderPassDescriptor];
+	[CurrentRenderPass retain];
 	TRACK_OBJECT(CurrentRenderPass);
 	// if we need to do queries, write to the ring buffer (we set the offset into the ring buffer per query)
 	CurrentRenderPass.visibilityResultBuffer = QueryBuffer.Buffer;
@@ -737,7 +755,6 @@ void FMetalManager::UpdateContext()
 		}
 
 		ReleaseObject(CurrentContext);
-///		ReleaseObject(CurrentFramebuffer);
 
 		// if we are doing occlusion queries, we could use this method, along with a completion callback
 		// to set a "render target complete" flag that the OQ code next frame would wait on
@@ -767,14 +784,6 @@ void FMetalManager::UpdateContext()
 		}];
 	}
 		
-	// make a new framebuffer object to render to
-///	MTLFramebufferDescriptor* FramebufferDescriptor = [MTLFramebufferDescriptor framebufferDescriptorWithColorAttachment:ColorAttachment depthAttachment:DepthAttachment stencilAttachment:nil];
-	// if we need to do queries, write to the ring buffer (we set the offset into the ring buffer per query)
-///	FramebufferDescriptor.visibilityResultBuffer = QueryBuffer.Buffer;
-
-///	CurrentFramebuffer = [Device newFramebufferWithDescriptor:FramebufferDescriptor];
-///	TRACK_OBJECT(CurrentFramebuffer);
-
 	// make a new render context to use to render to the framebuffer
 	CurrentContext = [CurrentCommandBuffer renderCommandEncoderWithDescriptor:CurrentRenderPass];
 	[CurrentContext retain];
@@ -819,7 +828,117 @@ uint32 FMetalManager::AllocateFromRingBuffer(uint32 Size, uint32 Alignment)
 
 uint32 FMetalManager::AllocateFromQueryBuffer()
 {
-	return QueryBuffer.Allocate(8, 0);
+ 	return QueryBuffer.Allocate(8, 0);
+}
+
+
+
+
+FORCEINLINE void SetResource(uint32 ShaderStage, uint32 BindIndex, FMetalSurface* RESTRICT Surface)
+{
+	check(Surface->Texture != nil);
+	if (ShaderStage == METAL_SHADER_STAGE_PIXEL)
+	{
+		[FMetalManager::GetContext() setFragmentTexture:Surface->Texture atIndex:BindIndex];
+	}
+	else
+	{
+		[FMetalManager::GetContext() setVertexTexture:Surface->Texture atIndex : BindIndex];
+	}
+}
+
+FORCEINLINE void SetResource(uint32 ShaderStage, uint32 BindIndex, FMetalSamplerState* RESTRICT SamplerState)
+{
+	check(SamplerState->State != nil);
+	if (ShaderStage == METAL_SHADER_STAGE_PIXEL)
+	{
+		[FMetalManager::GetContext() setFragmentSamplerState:SamplerState->State atIndex:BindIndex];
+	}
+	else
+	{
+		[FMetalManager::GetContext() setFragmentSamplerState:SamplerState->State atIndex:BindIndex];
+	}
+}
+
+template <typename MetalResourceType>
+inline int32 SetShaderResourcesFromBuffer(uint32 ShaderStage, FMetalUniformBuffer* RESTRICT Buffer, const uint32 * RESTRICT ResourceMap, int32 BufferIndex)
+{
+	int32 NumSetCalls = 0;
+	uint32 BufferOffset = ResourceMap[BufferIndex];
+	if (BufferOffset > 0)
+	{
+		const uint32* RESTRICT ResourceInfos = &ResourceMap[BufferOffset];
+		uint32 ResourceInfo = *ResourceInfos++;
+		do
+		{
+			checkSlow(FRHIResourceTableEntry::GetUniformBufferIndex(ResourceInfo) == BufferIndex);
+			const uint16 ResourceIndex = FRHIResourceTableEntry::GetResourceIndex(ResourceInfo);
+			const uint8 BindIndex = FRHIResourceTableEntry::GetBindIndex(ResourceInfo);
+
+			// @todo metal log: remove
+			NSLog(@"Binding a resource at BindIndex %d, Resource Index %d", BindIndex, ResourceIndex);
+
+			// todo: could coalesce adjacent bound resources.
+			MetalResourceType* ResourcePtr = (MetalResourceType*)Buffer->RawResourceTable[ResourceIndex];
+			SetResource(ShaderStage, BindIndex, ResourcePtr);
+
+			NumSetCalls++;
+			ResourceInfo = *ResourceInfos++;
+		} while (FRHIResourceTableEntry::GetUniformBufferIndex(ResourceInfo) == BufferIndex);
+	}
+
+	return NumSetCalls;
+}
+
+template <class ShaderType>
+void SetResourcesFromTables(ShaderType Shader, uint32 ShaderStage, uint32 ResourceTableFrameCounter)
+{
+	checkSlow(Shader);
+
+	// @todo metal log: remove
+	NSLog(@"Resource Bits %x, DirtyBits: %llx", Shader->Bindings.ShaderResourceTable.ResourceTableBits, Shader->DirtyUniformBuffers);
+
+	// Mask the dirty bits by those buffers from which the shader has bound resources.
+	uint32 DirtyBits = Shader->Bindings.ShaderResourceTable.ResourceTableBits & Shader->DirtyUniformBuffers;
+	uint32 NumSetCalls = 0;
+	while (DirtyBits)
+	{
+		// Scan for the lowest set bit, compute its index, clear it in the set of dirty bits.
+		const uint32 LowestBitMask = (DirtyBits)& (-(int32)DirtyBits);
+		const int32 BufferIndex = FMath::FloorLog2(LowestBitMask); // todo: This has a branch on zero, we know it could never be zero...
+		DirtyBits ^= LowestBitMask;
+		FMetalUniformBuffer* Buffer = (FMetalUniformBuffer*)Shader->BoundUniformBuffers[BufferIndex].GetReference();
+		check(Buffer);
+		check(BufferIndex < Shader->Bindings.ShaderResourceTable.ResourceTableLayoutHashes.Num());
+		check(Buffer->GetLayout().GetHash() == Shader->Bindings.ShaderResourceTable.ResourceTableLayoutHashes[BufferIndex]);
+		Buffer->CacheResources(ResourceTableFrameCounter);
+
+		// @todo metal log: remove
+		NSLog(@"Setting BufferIndex %d", BufferIndex);
+
+		// todo: could make this two pass: gather then set
+		NumSetCalls += SetShaderResourcesFromBuffer<FMetalSurface>(ShaderStage, Buffer, Shader->Bindings.ShaderResourceTable.ShaderResourceViewMap.GetData(), BufferIndex);
+		NumSetCalls += SetShaderResourcesFromBuffer<FMetalSamplerState>(ShaderStage, Buffer, Shader->Bindings.ShaderResourceTable.SamplerMap.GetData(), BufferIndex);
+	}
+	Shader->DirtyUniformBuffers = 0;
+//	SetTextureInTableCalls += NumSetCalls;
+}
+
+void FMetalManager::CommitGraphicsResourceTables()
+{
+	uint32 Start = FPlatformTime::Cycles();
+
+// 	GRHICommandList.Verify();
+
+	check(CurrentBoundShaderState);
+
+	SetResourcesFromTables(CurrentBoundShaderState->VertexShader, METAL_SHADER_STAGE_VERTEX, ResourceTableFrameCounter);
+	if (IsValidRef(CurrentBoundShaderState->PixelShader))
+	{
+		SetResourcesFromTables(CurrentBoundShaderState->PixelShader, METAL_SHADER_STAGE_PIXEL, ResourceTableFrameCounter);
+	}
+
+//	CommitResourceTableCycles += (FPlatformTime::Cycles() - Start);
 }
 
 void FMetalManager::CommitNonComputeShaderConstants()
@@ -908,122 +1027,3 @@ void FMetalManager::SetRasterizerState(const FRasterizerStateInitializerRHI& Sta
 
 	bFirstRasterizerState = false;
 }
-
-
-
-// @todo metal srt
-#if 0 
-template <class ShaderType>
-void FMetalManager::SetResourcesFromTables(const ShaderType* RESTRICT Shader, Gnm::ShaderStage ShaderStage)
-{
-	checkSlow(Shader);
-
-	// Mask the dirty bits by those buffers from which the shader has bound resources.
-	uint32 DirtyBits = Shader->ShaderResourceTable.ResourceTableBits & DirtyUniformBuffers[TRHIShaderToEnum<ShaderType>::ShaderFrequency];
-	uint32 NumSetCalls = 0;
-	while (DirtyBits)
-	{
-		// Scan for the lowest set bit, compute its index, clear it in the set of dirty bits.
-		const uint32 LowestBitMask = (DirtyBits) & (-(int32)DirtyBits);
-		const int32 BufferIndex = FMath::FloorLog2(LowestBitMask); // todo: This has a branch on zero, we know it could never be zero...
-		DirtyBits ^= LowestBitMask;
-		auto* Buffer = (FGnmUniformBuffer*)BoundUniformBuffers[TRHIShaderToEnum<ShaderType>::ShaderFrequency][BufferIndex].GetReference();
-		check(Buffer);
-		check(BufferIndex < Shader->ShaderResourceTable.ResourceTableLayoutHashes.Num());
-		check(Buffer->GetLayout().GetHash() == Shader->ShaderResourceTable.ResourceTableLayoutHashes[BufferIndex]);
-		Buffer->CacheResources(ResourceTableFrameCounter);
-
-		// todo: could make this two pass: gather then set
-		NumSetCalls +=
-		SetShaderResourcesFromBuffer<FGnmSurface>(ShaderStage, Buffer, Shader->ShaderResourceTable.ShaderResourceViewMap.GetData(), BufferIndex);
-		SetShaderResourcesFromBuffer<Gnm::Sampler>(ShaderStage, Buffer, Shader->ShaderResourceTable.SamplerMap.GetData(), BufferIndex);
-	}
-	DirtyUniformBuffers[TRHIShaderToEnum<ShaderType>::ShaderFrequency] = 0;
-	SetTextureInTableCalls += NumSetCalls;
-}
-
-void FMetalManager::CommitGraphicsResourceTables()
-{
-	uint32 Start = FPlatformTime::Cycles();
-
-	GRHICommandList.Verify();
-
-	auto* RESTRICT CurrentBoundShaderState = (FGnmBoundShaderState*)GGnmManager.BoundShaderStateHistory.GetLast();
-	check(CurrentBoundShaderState);
-
-	if (auto* Shader = CurrentBoundShaderState->VertexShader.GetReference())
-	{
-		SetResourcesFromTables(Shader, Shader->ShaderStage);
-	}
-	if (auto* Shader = CurrentBoundShaderState->PixelShader.GetReference())
-	{
-		SetResourcesFromTables(Shader, Gnm::kShaderStagePs);
-	}
-	if (auto* Shader = CurrentBoundShaderState->HullShader.GetReference())
-	{
-		SetResourcesFromTables(Shader, Gnm::kShaderStageHs);
-		auto* DomainShader = CurrentBoundShaderState->DomainShader.GetReference();
-		check(DomainShader);
-		check(0);
-/*
-		SetResourcesFromTables(DomainShader, Gnm::kShaderStageDs);
-*/
-	}
-	if (auto* Shader = CurrentBoundShaderState->GeometryShader.GetReference())
-	{
-		SetResourcesFromTables(Shader, Gnm::kShaderStageGs);
-	}
-
-	CommitResourceTableCycles += (FPlatformTime::Cycles() - Start);
-}
-
-void FMetalManager::CommitNonComputeShaderConstants()
-{
-	GRHICommandList.Verify();
-	auto* CurrentBoundShaderState = (FGnmBoundShaderState*)BoundShaderStateHistory.GetLast();
-	check(CurrentBoundShaderState);
-
-	// update constant buffers from local shadow memory, and send to GPU
-	if (CurrentBoundShaderState->bShaderNeedsGlobalConstantBuffer[ SF_Vertex ] )
-	{
-		VSConstantBuffer->CommitConstantsToDevice(CurrentVertexShaderStage, 0, bDiscardSharedConstants);
-	}
-	if (CurrentBoundShaderState->bShaderNeedsGlobalConstantBuffer[ SF_Pixel ] )
-	{
-		PSConstantBuffer->CommitConstantsToDevice(Gnm::kShaderStagePs, 0, bDiscardSharedConstants);
-	}
-	if (CurrentBoundShaderState->bShaderNeedsGlobalConstantBuffer[ SF_Geometry ] )
-	{
-		GSConstantBuffer->CommitConstantsToDevice(Gnm::kShaderStageGs, 0, bDiscardSharedConstants);
-	}
-/*
-	// Skip HS/DS CB updates in cases where tessellation isn't being used
-	// Note that this is *potentially* unsafe because bDiscardSharedConstants is cleared at the
-	// end of the function, however we're OK for now because bDiscardSharedConstants
-	// is always reset whenever bUsingTessellation changes in SetBoundShaderState()
-	if (bUsingTessellation)
-	{
-		if (CurrentBoundShaderState->bShaderNeedsGlobalConstantBuffer[ SF_Hull ] )
-		{
-			HSConstantBuffer->CommitConstantsToDevice(Gnm::kShaderStageHs, 0, bDiscardSharedConstants);
-		}
-		if (CurrentBoundShaderState->bShaderNeedsGlobalConstantBuffer[ SF_Domain ] )
-		{
-			check(0);
-			//DSConstantBuffer->CommitConstantsToDevice(Gnm::kShaderStageDs, 0, bDiscardSharedConstants);
-		}
-	}
-*/
-
-	// don't need to discard again until next shader is set
-	bDiscardSharedConstants = false;
-}
-
-void FGnmManager::CommitComputeResourceTables( FGnmComputeShader* ComputeShader )
-{
-	GRHICommandList.Verify();
-
-	check(ComputeShader);
-	SetResourcesFromTables(ComputeShader, Gnm::kShaderStageCs);
-}
-#endif
\ No newline at end of file
diff --git a/Engine/Source/Runtime/IOS/MetalRHI/Private/MetalManager.h b/Engine/Source/Runtime/IOS/MetalRHI/Private/MetalManager.h
--- a/Engine/Source/Runtime/IOS/MetalRHI/Private/MetalManager.h
+++ b/Engine/Source/Runtime/IOS/MetalRHI/Private/MetalManager.h
@@ -73,8 +73,13 @@ public:
 	static id<MTLRenderCommandEncoder> GetContext();
 	static void ReleaseObject(id Object);
 	
+	/** RHIBeginScene helper */
+	void BeginScene();
+	/** RHIEndScene helper */
+	void EndScene();
+
 	void BeginFrame();
-	void EndFrame(bool bPreset);
+	void EndFrame(bool bPresent);
 	
 	/**
 	 * Return a texture object that wraps the backbuffer
@@ -134,6 +139,11 @@ public:
 
 	void SetRasterizerState(const FRasterizerStateInitializerRHI& State);
 
+	FMetalShaderParameterCache& GetShaderParameters(uint32 Stage)
+	{
+		return ShaderParameters[Stage];
+	}
+
 protected:
 	FMetalManager();
 	void InitFrame();
@@ -159,6 +169,8 @@ protected:
 	TRefCountPtr<FMetalTexture2D> BackBuffer;
 	TRefCountPtr<FMetalBoundShaderState> CurrentBoundShaderState;
 	
+	class FMetalShaderParameterCache*	ShaderParameters;
+
 	// the running pipeline state descriptor object
 	FPipelineShadow Pipeline;
 
@@ -180,6 +192,15 @@ protected:
 	uint64 CommandBufferIndex;
 	uint64 CompletedCommandBufferIndex;
 
+	/** Internal frame counter, incremented on each call to RHIBeginScene. */
+	uint32 SceneFrameCounter;
+
+	/**
+	 * Internal counter used for resource table caching.
+	 * INDEX_NONE means caching is not allowed.
+	 */
+	uint32 ResourceTableFrameCounter;
+
 	// shadow the rasterizer state to reduce some render encoder pressure
 	FRasterizerStateInitializerRHI ShadowRasterizerState;
 	bool bFirstRasterizerState;
@@ -193,10 +214,11 @@ protected:
 	uint8 LastIndexToDump;
 #endif
 
-public:
+	/** Apply the SRT before drawing */
+	void CommitGraphicsResourceTables();
+
 	void CommitNonComputeShaderConstants();
-	
-	class FMetalShaderParameterCache*	ShaderParameters;
+
 };
 
 // Stats
diff --git a/Engine/Source/Runtime/IOS/MetalRHI/Private/MetalRHI.cpp b/Engine/Source/Runtime/IOS/MetalRHI/Private/MetalRHI.cpp
--- a/Engine/Source/Runtime/IOS/MetalRHI/Private/MetalRHI.cpp
+++ b/Engine/Source/Runtime/IOS/MetalRHI/Private/MetalRHI.cpp
@@ -163,12 +163,12 @@ void FMetalDynamicRHI::RHIEndFrame()
 
 void FMetalDynamicRHI::RHIBeginScene()
 {
-
+	FMetalManager::Get()->BeginScene();
 }
 
 void FMetalDynamicRHI::RHIEndScene()
 {
-
+	FMetalManager::Get()->EndScene();
 }
 
 void FMetalDynamicRHI::PushEvent(const TCHAR* Name)
diff --git a/Engine/Source/Runtime/IOS/MetalRHI/Private/MetalShaders.cpp b/Engine/Source/Runtime/IOS/MetalRHI/Private/MetalShaders.cpp
--- a/Engine/Source/Runtime/IOS/MetalRHI/Private/MetalShaders.cpp
+++ b/Engine/Source/Runtime/IOS/MetalRHI/Private/MetalShaders.cpp
@@ -58,6 +58,7 @@ static FMetalCompiledShaderCache& GetMetalCompiledShaderCache()
 /** Initialization constructor. */
 template<typename BaseResourceType>
 TMetalBaseShader<BaseResourceType>::TMetalBaseShader(const TArray<uint8>& InCode)
+	: DirtyUniformBuffers(0)
 {
 	FMemoryReader Ar(InCode, true);
 
diff --git a/Engine/Source/Runtime/IOS/MetalRHI/Private/MetalUniformBuffer.cpp b/Engine/Source/Runtime/IOS/MetalRHI/Private/MetalUniformBuffer.cpp
--- a/Engine/Source/Runtime/IOS/MetalRHI/Private/MetalUniformBuffer.cpp
+++ b/Engine/Source/Runtime/IOS/MetalRHI/Private/MetalUniformBuffer.cpp
@@ -293,7 +293,7 @@ void FMetalUniformBuffer::CacheResourcesInternal()
 				break;
 
 			case UBMT_SAMPLER:
-				RawResources[i] = &((FMetalSamplerState*)Resources[i].GetReference())->State;
+				RawResources[i] = (FMetalSamplerState*)Resources[i].GetReference();
 				break;
 
 			default:
diff --git a/Engine/Source/Runtime/IOS/MetalRHI/Public/MetalResources.h b/Engine/Source/Runtime/IOS/MetalRHI/Public/MetalResources.h
--- a/Engine/Source/Runtime/IOS/MetalRHI/Public/MetalResources.h
+++ b/Engine/Source/Runtime/IOS/MetalRHI/Public/MetalResources.h
@@ -67,6 +67,10 @@ public:
  	/** External bindings for this shader. */
 	FMetalShaderBindings Bindings;
 	TArray< TRefCountPtr<FRHIUniformBuffer> > BoundUniformBuffers;
+
+	/** Bitfield for which uniform buffers are dirty */
+	uint64 DirtyUniformBuffers;
+
 	// List of memory copies from RHIUniformBuffer to packed uniforms
 	TArray<FMetalUniformBufferCopyInfo> UniformBuffersCopyInfo;
     
diff --git a/Engine/Source/Runtime/SlateRHIRenderer/Private/SlateRHIRenderer.cpp b/Engine/Source/Runtime/SlateRHIRenderer/Private/SlateRHIRenderer.cpp
--- a/Engine/Source/Runtime/SlateRHIRenderer/Private/SlateRHIRenderer.cpp
+++ b/Engine/Source/Runtime/SlateRHIRenderer/Private/SlateRHIRenderer.cpp
@@ -423,6 +423,10 @@ void FSlateRHIRenderer::DrawWindow_RenderThread( const FViewportInfo& ViewportIn
 			// Clear the stencil buffer
 			RHIClear( false, FLinearColor::White, false, 0.0f, true, 0x00, FIntRect());
 		}
+		else
+		{
+			RHISetRenderTarget(BackBuffer, FTextureRHIRef());
+		}
 
 #if DEBUG_OVERDRAW
 		RHIClear(true, FLinearColor::Black, false, 0.0f, true, 0x00, FIntRect());
