UE4 - Metal compile fix

diff --git a/Engine/Source/Runtime/IOS/MetalRHI/Private/MetalRHIPrivate.h b/Engine/Source/Runtime/IOS/MetalRHI/Private/MetalRHIPrivate.h
--- a/Engine/Source/Runtime/IOS/MetalRHI/Private/MetalRHIPrivate.h
+++ b/Engine/Source/Runtime/IOS/MetalRHI/Private/MetalRHIPrivate.h
@@ -56,19 +56,19 @@ extern TMap<id, int32> ClassCounts;
 class FLocalEvent : public FEvent
 {
 	// This is a little complicated, in an attempt to match Win32 Event semantics...
-    typedef enum
-    {
-        TRIGGERED_NONE,
-        TRIGGERED_ONE,
-        TRIGGERED_ALL,
-    } TriggerType;
+	typedef enum
+	{
+		TRIGGERED_NONE,
+		TRIGGERED_ONE,
+		TRIGGERED_ALL,
+	} TriggerType;
 	
 	bool bInitialized;
 	bool bIsManualReset;
 	volatile TriggerType Triggered;
 	volatile int32 WaitingThreads;
-    pthread_mutex_t Mutex;
-    pthread_cond_t Condition;
+	pthread_mutex_t Mutex;
+	pthread_cond_t Condition;
 	
 	inline void LockEventMutex()
 	{
@@ -194,7 +194,7 @@ public:
 		UnlockEventMutex();
 	}
 	
-	virtual bool Wait(uint32 WaitTime)
+	virtual bool Wait(uint32 WaitTime, const bool bIgnoreThreadIdleStats = false) override
 	{
 		check(bInitialized);
 		
