Metal is now rendering, but overbright and slow

diff --git a/Engine/Source/Runtime/IOS/MetalRHI/Private/MetalCommands.cpp b/Engine/Source/Runtime/IOS/MetalRHI/Private/MetalCommands.cpp
--- a/Engine/Source/Runtime/IOS/MetalRHI/Private/MetalCommands.cpp
+++ b/Engine/Source/Runtime/IOS/MetalRHI/Private/MetalCommands.cpp
@@ -261,7 +261,7 @@ void FMetalDynamicRHI::RHISetShaderUniformBuffer(FVertexShaderRHIParamRef Vertex
 	VertexShader->DirtyUniformBuffers |= 1 << BufferIndex;
 
 	// @todo metal log: remove
-	NSLog(@"Setting Vertex UB at index %d", BufferIndex);
+//	NSLog(@"Setting Vertex UB at index %d", BufferIndex);
 }
 
 void FMetalDynamicRHI::RHISetShaderUniformBuffer(FHullShaderRHIParamRef HullShader, uint32 BufferIndex, FUniformBufferRHIParamRef BufferRHI)
@@ -286,7 +286,7 @@ void FMetalDynamicRHI::RHISetShaderUniformBuffer(FPixelShaderRHIParamRef PixelSh
 	PixelShader->DirtyUniformBuffers |= 1 << BufferIndex;
 
 	// @todo metal log: remove
-	NSLog(@"Setting Pixel UB at index %d", BufferIndex);
+//	NSLog(@"Setting Pixel UB at index %d", BufferIndex);
 }
 
 void FMetalDynamicRHI::RHISetShaderUniformBuffer(FComputeShaderRHIParamRef ComputeShader, uint32 BufferIndex, FUniformBufferRHIParamRef BufferRHI)
diff --git a/Engine/Source/Runtime/IOS/MetalRHI/Private/MetalManager.cpp b/Engine/Source/Runtime/IOS/MetalRHI/Private/MetalManager.cpp
--- a/Engine/Source/Runtime/IOS/MetalRHI/Private/MetalManager.cpp
+++ b/Engine/Source/Runtime/IOS/MetalRHI/Private/MetalManager.cpp
@@ -876,7 +876,7 @@ inline int32 SetShaderResourcesFromBuffer(uint32 ShaderStage, FMetalUniformBuffe
 			const uint8 BindIndex = FRHIResourceTableEntry::GetBindIndex(ResourceInfo);
 
 			// @todo metal log: remove
-			NSLog(@"Binding a resource at BindIndex %d, Resource Index %d", BindIndex, ResourceIndex);
+//			NSLog(@"Binding a resource at BindIndex %d, Resource Index %d", BindIndex, ResourceIndex);
 
 			// todo: could coalesce adjacent bound resources.
 			MetalResourceType* ResourcePtr = (MetalResourceType*)Buffer->RawResourceTable[ResourceIndex];
@@ -896,7 +896,7 @@ void SetResourcesFromTables(ShaderType Shader, uint32 ShaderStage, uint32 Resour
 	checkSlow(Shader);
 
 	// @todo metal log: remove
-	NSLog(@"Resource Bits %x, DirtyBits: %llx", Shader->Bindings.ShaderResourceTable.ResourceTableBits, Shader->DirtyUniformBuffers);
+//	NSLog(@"Resource Bits %x, DirtyBits: %llx", Shader->Bindings.ShaderResourceTable.ResourceTableBits, Shader->DirtyUniformBuffers);
 
 	// Mask the dirty bits by those buffers from which the shader has bound resources.
 	uint32 DirtyBits = Shader->Bindings.ShaderResourceTable.ResourceTableBits & Shader->DirtyUniformBuffers;
@@ -914,10 +914,11 @@ void SetResourcesFromTables(ShaderType Shader, uint32 ShaderStage, uint32 Resour
 		Buffer->CacheResources(ResourceTableFrameCounter);
 
 		// @todo metal log: remove
-		NSLog(@"Setting BufferIndex %d", BufferIndex);
+//		NSLog(@"Setting BufferIndex %d", BufferIndex);
 
 		// todo: could make this two pass: gather then set
-		NumSetCalls += SetShaderResourcesFromBuffer<FMetalSurface>(ShaderStage, Buffer, Shader->Bindings.ShaderResourceTable.ShaderResourceViewMap.GetData(), BufferIndex);
+//		NumSetCalls += SetShaderResourcesFromBuffer<FMetalSurface>(ShaderStage, Buffer, Shader->Bindings.ShaderResourceTable.ShaderResourceViewMap.GetData(), BufferIndex);
+		NumSetCalls += SetShaderResourcesFromBuffer<FMetalSurface>(ShaderStage, Buffer, Shader->Bindings.ShaderResourceTable.TextureMap.GetData(), BufferIndex);
 		NumSetCalls += SetShaderResourcesFromBuffer<FMetalSamplerState>(ShaderStage, Buffer, Shader->Bindings.ShaderResourceTable.SamplerMap.GetData(), BufferIndex);
 	}
 	Shader->DirtyUniformBuffers = 0;
diff --git a/Engine/Source/Runtime/IOS/MetalRHI/Private/MetalUniformBuffer.cpp b/Engine/Source/Runtime/IOS/MetalRHI/Private/MetalUniformBuffer.cpp
--- a/Engine/Source/Runtime/IOS/MetalRHI/Private/MetalUniformBuffer.cpp
+++ b/Engine/Source/Runtime/IOS/MetalRHI/Private/MetalUniformBuffer.cpp
@@ -181,6 +181,7 @@ FMetalUniformBuffer::FMetalUniformBuffer(const void* Contents, const FRHIUniform
 	: FRHIUniformBuffer(Layout)
 	, Buffer(nil)
 	, Offset(0)
+	, LastCachedFrame(INDEX_NONE)
 {
 	if (Layout.ConstantBufferSize > 0)
 	{
@@ -263,11 +264,12 @@ void FMetalUniformBuffer::CacheResourcesInternal()
 		{
 			case UBMT_SRV:
 				{
+					check(0);
 					FMetalShaderResourceView* SRV = (FMetalShaderResourceView*)Resources[i].GetReference();
 					if (IsValidRef(SRV->SourceTexture))
 					{
 						FMetalSurface& Surface = GetMetalSurfaceFromRHITexture(SRV->SourceTexture);
-						RawResources[i] = Surface.Texture;
+						RawResources[i] = &Surface;
 					}
 					else
 					{
