- Removed the Zero Stride buffer hack from Metal now that it's supported properly by Apple
- Removed some other hackery

diff --git a/Engine/Source/Runtime/IOS/MetalRHI/Private/MetalCommands.cpp b/Engine/Source/Runtime/IOS/MetalRHI/Private/MetalCommands.cpp
--- a/Engine/Source/Runtime/IOS/MetalRHI/Private/MetalCommands.cpp
+++ b/Engine/Source/Runtime/IOS/MetalRHI/Private/MetalCommands.cpp
@@ -19,15 +19,10 @@ static MTLPrimitiveType TranslatePrimitiveType(uint32 PrimitiveType)
 	}
 }
 
-TRefCountPtr<FMetalVertexBuffer> HACK_ZeroStrideBuffers[16] = { 0 };
-
 void FMetalDynamicRHI::RHISetStreamSource(uint32 StreamIndex,FVertexBufferRHIParamRef VertexBufferRHI,uint32 Stride,uint32 Offset)
 {
 	DYNAMIC_CAST_METGALRESOURCE(VertexBuffer,VertexBuffer);
 
-	// queue this zero stride buffer to be grown later if needed
-	HACK_ZeroStrideBuffers[StreamIndex] = VertexBuffer->ZeroStrideElementSize ? VertexBuffer : NULL;
-
 	[FMetalManager::GetContext() setVertexBuffer:VertexBuffer->Buffer offset:VertexBuffer->Offset + Offset atIndex:UNREAL_TO_METAL_BUFFER_INDEX(StreamIndex)];
 }
 
@@ -479,7 +474,6 @@ void FMetalDynamicRHI::RHIEndDrawPrimitiveUP()
 	uint32 NumVertices = RHIGetVertexCountForPrimitiveCount(GPendingNumPrimitives, GPendingPrimitiveType);
 
 	// last minute draw setup
-	HACK_ZeroStrideBuffers[0] = NULL;
 	FMetalManager::Get()->PrepareToDraw(0);
 	
 #if !NO_DRAW
@@ -525,7 +519,6 @@ void FMetalDynamicRHI::RHIEndDrawIndexedPrimitiveUP()
 	uint32 NumIndices = RHIGetVertexCountForPrimitiveCount(GPendingNumPrimitives, GPendingPrimitiveType);
 
 	// last minute draw setup
-	HACK_ZeroStrideBuffers[0] = NULL;
 	FMetalManager::Get()->PrepareToDraw(0);
 	
 #if !NO_DRAW
diff --git a/Engine/Source/Runtime/IOS/MetalRHI/Private/MetalManager.cpp b/Engine/Source/Runtime/IOS/MetalRHI/Private/MetalManager.cpp
--- a/Engine/Source/Runtime/IOS/MetalRHI/Private/MetalManager.cpp
+++ b/Engine/Source/Runtime/IOS/MetalRHI/Private/MetalManager.cpp
@@ -510,39 +510,6 @@ void FMetalManager::PrepareToDraw(uint32 NumVertices)
 	
 	CommitGraphicsResourceTables();
 	CommitNonComputeShaderConstants();
-
-	// grow buffers as needed
-	extern TRefCountPtr<FMetalVertexBuffer> HACK_ZeroStrideBuffers[16];
-	for (int32 Index = 0; Index < 16; Index++)
-	{
-		if (HACK_ZeroStrideBuffers[Index] != NULL)
-		{
-			TRefCountPtr<FMetalVertexBuffer> Buffer = HACK_ZeroStrideBuffers[Index];
-			
-			// how big do we need?
-			uint32 NeededSize = Align(NumVertices * Buffer->ZeroStrideElementSize, 128*1024);
-			// if too small, then grow and fill in
-			if (NeededSize > [Buffer->Buffer length])
-			{
-				id<MTLBuffer> Orig = Buffer->Buffer;
-				Buffer->Buffer = [Device newBufferWithLength:NeededSize options:BUFFER_CACHE_MODE];
-				TRACK_OBJECT(Buffer->Buffer);
-
-				uint8* SrcMem = (uint8*)[Orig contents];
-				uint8* DestMem = (uint8*)[Buffer->Buffer contents];
-				for (int32 CopyIndex = 1; CopyIndex < NumVertices; CopyIndex++)
-				{
-					FMemory::Memcpy(DestMem + Buffer->ZeroStrideElementSize * CopyIndex, SrcMem, Buffer->ZeroStrideElementSize);
-				}
-
-				NSLog(@"Expaning buffer from %d to %d", (int32)[Orig length], NeededSize);
-				// set the new buffer object
-				[FMetalManager::GetContext() setVertexBuffer:Buffer->Buffer offset:0 atIndex:UNREAL_TO_METAL_BUFFER_INDEX(Index)];
-
-				ReleaseObject(Orig);
-			}
-		}
-	}
 }
 
 void FMetalManager::SetDepthStencilWriteEnabled(bool bIsDepthWriteEnabled, bool bIsStencilWriteEnabled)
diff --git a/Engine/Source/Runtime/IOS/MetalRHI/Private/MetalShaders.cpp b/Engine/Source/Runtime/IOS/MetalRHI/Private/MetalShaders.cpp
--- a/Engine/Source/Runtime/IOS/MetalRHI/Private/MetalShaders.cpp
+++ b/Engine/Source/Runtime/IOS/MetalRHI/Private/MetalShaders.cpp
@@ -11,14 +11,6 @@
 /** Set to 1 to enable shader debugging (makes the driver save the shader source) */
 #define DEBUG_METAL_SHADERS (UE_BUILD_DEBUG || UE_BUILD_DEVELOPMENT)
 
-uint64 HitchyHashes[] = 
-{
- 	0x09f873f002005ULL,
- 	0x040050e020150ULL,
- 	0x09f873e0a0150ULL,
-	0x11f873f002005ULL
-};
-
 
 class FMetalCompiledShaderKey
 {
@@ -269,19 +261,6 @@ FMetalBoundShaderState::FMetalBoundShaderState(
 	}
 #endif
 	
-	if (GFrameCounter < 2000)
-	{
-		for (int32 Index = 0; Index < ARRAY_COUNT(HitchyHashes); Index++)
-		{
-			FPipelineShadow Shadow;
-			Shadow.SetHash(HitchyHashes[Index]);
-			check(Shadow.GetHash() == HitchyHashes[Index]);
-	 
-			id<MTLRenderPipelineState> PipelineState = Shadow.CreatePipelineStateForBoundShaderState(this);
-			PipelineStates.Add(HitchyHashes[Index], PipelineState);
-		}
-	}
-	
 	AddRef();
 }
 
@@ -297,10 +276,6 @@ FMetalBoundShaderState::~FMetalBoundShaderState()
 
 void FMetalBoundShaderState::PrepareToDraw(const FPipelineShadow& PipelineShadow)
 {
-#if NO_DRAW
-	return;
-#endif
-	
 	// generate a key for the current state
 	uint64 Hash = PipelineShadow.GetHash();
 	
