UE4 - Metal - temp fix assert to unblock... need to properly track this down!

diff --git a/Engine/Source/Runtime/IOS/MetalRHI/Private/MetalShaders.cpp b/Engine/Source/Runtime/IOS/MetalRHI/Private/MetalShaders.cpp
--- a/Engine/Source/Runtime/IOS/MetalRHI/Private/MetalShaders.cpp
+++ b/Engine/Source/Runtime/IOS/MetalRHI/Private/MetalShaders.cpp
@@ -279,7 +279,7 @@ void FMetalBoundShaderState::PrepareToDraw(const FPipelineShadow& PipelineShadow
 	// generate a key for the current state
 	uint64 Hash = PipelineShadow.GetHash();
 	
-    
+	
 	// have we made a matching state object yet?
 	id<MTLRenderPipelineState> PipelineState = PipelineStates.FindRef(Hash);
 	
@@ -422,16 +422,17 @@ void FMetalShaderParameterCache::CommitPackedGlobals(int32 Stage, const FMetalSh
 	// copy the current uniform buffer into the ring buffer to submit
 	for (int32 Index = 0; Index < Bindings.PackedGlobalArrays.Num(); ++Index)
 	{
-        int32 UniformBufferIndex = Bindings.PackedGlobalArrays[Index].TypeIndex;
+		int32 UniformBufferIndex = Bindings.PackedGlobalArrays[Index].TypeIndex;
  
 		// is there any data that needs to be copied?
 		if (PackedGlobalUniformDirty[UniformBufferIndex].HighVector > 0)//PackedGlobalUniformDirty[UniformBufferIndex].LowVector)
 		{
-  			check(UniformBufferIndex == 0 || UniformBufferIndex == 1);
-            uint32 TotalSize = Bindings.PackedGlobalArrays[Index].Size;
+			check(UniformBufferIndex == 0 || UniformBufferIndex == 1);
+			uint32 TotalSize = Bindings.PackedGlobalArrays[Index].Size;
 			uint32 SizeToUpload = PackedGlobalUniformDirty[UniformBufferIndex].HighVector * SizeOfFloat4;
-            check(SizeToUpload <= TotalSize);
-            SizeToUpload = TotalSize;
+//@todo-rco: Temp workaround
+			//check(SizeToUpload <= TotalSize);
+			SizeToUpload = TotalSize;
  /*
 			if (UniformBufferIndex == METAL_PACKED_TYPEINDEX_MEDIUMP)
 			{
@@ -462,7 +463,8 @@ void FMetalShaderParameterCache::CommitPackedGlobals(int32 Stage, const FMetalSh
  */
 			{
 				//check(UniformBufferIndex != METAL_PACKED_TYPEINDEX_LOWP);
-				FMemory::Memcpy(((uint8*)[RingBuffer contents]) + Offset, PackedGlobalUniforms[UniformBufferIndex], SizeToUpload);
+//@todo-rco: Temp workaround
+				FMemory::Memcpy(((uint8*)[RingBuffer contents]) + Offset, PackedGlobalUniforms[UniformBufferIndex], FMath::Min(TotalSize, SizeToUpload));
 			}
 			
 			switch (Stage)
