- Added GRHISupportsBaseVertexIndex for slate on Metal, where BaseVertexIndex is not quickly/easily supported, and a high level code change makes it work
- Metal RHI cleanup

diff --git a/Engine/Source/Runtime/IOS/MetalRHI/Private/MetalCommands.cpp b/Engine/Source/Runtime/IOS/MetalRHI/Private/MetalCommands.cpp
--- a/Engine/Source/Runtime/IOS/MetalRHI/Private/MetalCommands.cpp
+++ b/Engine/Source/Runtime/IOS/MetalRHI/Private/MetalCommands.cpp
@@ -6,8 +6,6 @@
 
 #include "MetalRHIPrivate.h"
 
-#define NO_DRAW 0
-
 static MTLPrimitiveType TranslatePrimitiveType(uint32 PrimitiveType)
 {
 	switch (PrimitiveType)
@@ -254,9 +252,6 @@ void FMetalDynamicRHI::RHISetShaderUniformBuffer(FVertexShaderRHIParamRef Vertex
 	DYNAMIC_CAST_METGALRESOURCE(VertexShader, VertexShader);
 	VertexShader->BoundUniformBuffers[BufferIndex] = BufferRHI;
 	VertexShader->DirtyUniformBuffers |= 1 << BufferIndex;
-
-	// @todo metal log: remove
-//	NSLog(@"Setting Vertex UB at index %d", BufferIndex);
 }
 
 void FMetalDynamicRHI::RHISetShaderUniformBuffer(FHullShaderRHIParamRef HullShader, uint32 BufferIndex, FUniformBufferRHIParamRef BufferRHI)
@@ -279,9 +274,6 @@ void FMetalDynamicRHI::RHISetShaderUniformBuffer(FPixelShaderRHIParamRef PixelSh
 	DYNAMIC_CAST_METGALRESOURCE(PixelShader, PixelShader);
 	PixelShader->BoundUniformBuffers[BufferIndex] = BufferRHI;
 	PixelShader->DirtyUniformBuffers |= 1 << BufferIndex;
-
-	// @todo metal log: remove
-//	NSLog(@"Setting Pixel UB at index %d", BufferIndex);
 }
 
 void FMetalDynamicRHI::RHISetShaderUniformBuffer(FComputeShaderRHIParamRef ComputeShader, uint32 BufferIndex, FUniformBufferRHIParamRef BufferRHI)
@@ -371,12 +363,6 @@ void FMetalDynamicRHI::RHIDrawPrimitive(uint32 PrimitiveType, uint32 BaseVertexI
 	SCOPE_CYCLE_COUNTER(STAT_MetalDrawCallTime);
 	checkf(NumInstances == 1, TEXT("Currently only 1 instance is supported"));
 	
-	if (BaseVertexIndex > 0)
-	{
-		UE_LOG(LogMetal, Log, TEXT("Attempted to render with BaseVertexIndex - skipping..."));
-		return;
-	}
-
 	// how many verts to render
 	uint32 NumVertices = GetVertexCountForPrimitiveCount(NumPrimitives, PrimitiveType);
 
@@ -384,12 +370,10 @@ void FMetalDynamicRHI::RHIDrawPrimitive(uint32 PrimitiveType, uint32 BaseVertexI
 	FMetalManager::Get()->PrepareToDraw(NumVertices);
 
 	// draw!
-#if !NO_DRAW
 	[FMetalManager::GetContext() drawPrimitives:TranslatePrimitiveType(PrimitiveType)
 									vertexStart:BaseVertexIndex
 									vertexCount:NumVertices
 								  instanceCount:NumInstances];
-#endif
 }
 
 void FMetalDynamicRHI::RHIDrawPrimitiveIndirect(uint32 PrimitiveType, FVertexBufferRHIParamRef ArgumentBufferRHI, uint32 ArgumentOffset)
@@ -403,11 +387,7 @@ void FMetalDynamicRHI::RHIDrawIndexedPrimitive(FIndexBufferRHIParamRef IndexBuff
 {
 	SCOPE_CYCLE_COUNTER(STAT_MetalDrawCallTime);
 	checkf(NumInstances == 1, TEXT("Currently only 1 instance is supported"));
-	if (BaseVertexIndex > 0)
-	{
-		UE_LOG(LogMetal, Log, TEXT("Attempted to render with BaseVertexIndex - skipping..."));
-		return;
-	}
+	checkf(BaseVertexIndex  == 0, TEXT("BaseVertexIndex must be 0, see GRHISupportsBaseVertexIndex"));
 
 	DYNAMIC_CAST_METGALRESOURCE(IndexBuffer,IndexBuffer);
 
@@ -415,14 +395,13 @@ void FMetalDynamicRHI::RHIDrawIndexedPrimitive(FIndexBufferRHIParamRef IndexBuff
 	FMetalManager::Get()->PrepareToDraw(NumVertices);
 	
 	uint32 NumIndices = GetVertexCountForPrimitiveCount(NumPrimitives, PrimitiveType);
-#if !NO_DRAW
+
 	[FMetalManager::GetContext() drawIndexedPrimitives:TranslatePrimitiveType(PrimitiveType)
 											indexCount:NumIndices
 											 indexType:IndexBuffer->IndexType
 										   indexBuffer:IndexBuffer->Buffer
 									 indexBufferOffset:IndexBuffer->Offset + StartIndex * IndexBuffer->GetStride()
 										 instanceCount:NumInstances];
-#endif
 }
 
 void FMetalDynamicRHI::RHIDrawIndexedIndirect(FIndexBufferRHIParamRef IndexBufferRHI, uint32 PrimitiveType, FStructuredBufferRHIParamRef ArgumentsBufferRHI, int32 DrawArgumentsIndex, uint32 NumInstances)
@@ -476,12 +455,10 @@ void FMetalDynamicRHI::RHIEndDrawPrimitiveUP()
 	// last minute draw setup
 	FMetalManager::Get()->PrepareToDraw(0);
 	
-#if !NO_DRAW
 	[FMetalManager::GetContext() drawPrimitives:TranslatePrimitiveType(GPendingPrimitiveType)
 									vertexStart:0
 									vertexCount:NumVertices
 								  instanceCount:1];
-#endif
 	
 	// mark temp memory as usable
 	GPendingVertexBufferOffset = 0xFFFFFFFF;
@@ -521,14 +498,12 @@ void FMetalDynamicRHI::RHIEndDrawIndexedPrimitiveUP()
 	// last minute draw setup
 	FMetalManager::Get()->PrepareToDraw(0);
 	
-#if !NO_DRAW
 	[FMetalManager::GetContext() drawIndexedPrimitives:TranslatePrimitiveType(GPendingPrimitiveType)
 											indexCount:NumIndices
 											 indexType:(GPendingIndexDataStride == 2) ? MTLIndexTypeUInt16 : MTLIndexTypeUInt32
 										   indexBuffer:FMetalManager::Get()->GetRingBuffer()
 									 indexBufferOffset:GPendingIndexBufferOffset
 										 instanceCount:1];
-#endif
 	
 	// mark temp memory as usable
 	GPendingVertexBufferOffset = 0xFFFFFFFF;
diff --git a/Engine/Source/Runtime/IOS/MetalRHI/Private/MetalManager.cpp b/Engine/Source/Runtime/IOS/MetalRHI/Private/MetalManager.cpp
--- a/Engine/Source/Runtime/IOS/MetalRHI/Private/MetalManager.cpp
+++ b/Engine/Source/Runtime/IOS/MetalRHI/Private/MetalManager.cpp
@@ -282,13 +282,6 @@ FMetalManager::FMetalManager()
 
 	AutoReleasePoolTLSSlot = FPlatformTLS::AllocTlsSlot();
 
-#if GPU_TIMINGS
-	// start at start of GPUTimers
-	TimerIndex = 0;
-	FirstIndexToDump = 0;
-	LastIndexToDump = 0;
-#endif
-	
 	// Create display link thread
 	FramePacer = [[MetalFramePacer alloc] init];
 	[NSThread detachNewThreadSelector:@selector(run:) toTarget:FramePacer withObject:nil];
@@ -338,19 +331,9 @@ void FMetalManager::InitFrame()
 	[CurrentCommandBuffer retain];
 	TRACK_OBJECT(CurrentCommandBuffer);
 	
-#if GPU_TIMINGS
-	uint8 LocalTimerIndex = TimerIndex++;
-	FirstIndexThisFrame = LocalTimerIndex;
-#endif
-
 	uint64 LocalCommandBufferIndex = CommandBufferIndex++;
 	[CurrentCommandBuffer addScheduledHandler : ^ (id <MTLCommandBuffer> Buffer)
 	{
-#if GPU_TIMINGS
-		// timer stuff
-		GPUTimers[LocalTimerIndex] = FPlatformTime::Seconds();
-#endif
-
 		FMetalManager::Get()->SetCompletedCommandBufferIndex(LocalCommandBufferIndex);
 	}];
 
@@ -393,21 +376,10 @@ void FMetalManager::EndFrame(bool bPresent)
 	ReleaseObject(CurrentContext);
 	CurrentContext = nil;
 
-#if GPU_TIMINGS
-	// set up timer
-	uint8 LocalTimerIndex = TimerIndex++;
-	uint8 LocalFirstIndex = FirstIndexThisFrame;
-#endif
-
 	// kick the whole buffer
 	[CurrentCommandBuffer addCompletedHandler:^(id <MTLCommandBuffer> Buffer)
 	 {
-#if GPU_TIMINGS
-	 GPUTimers[LocalTimerIndex] = FPlatformTime::Seconds();
-	 FirstIndexToDump = LocalFirstIndex;
-	 LastIndexToDump = LocalTimerIndex;
-#endif
-	 dispatch_semaphore_signal(CommandBufferSemaphore);
+		 dispatch_semaphore_signal(CommandBufferSemaphore);
 	 }];
 
 	// Wait until at least 2 VBlanks has passed since last time
@@ -439,23 +411,6 @@ void FMetalManager::EndFrame(bool bPresent)
 	UNTRACK_OBJECT(CurrentCommandBuffer);
 	[CurrentCommandBuffer release];
 
-#if GPU_TIMINGS
-	// print out timings
-	bool bDone = false;
-	uint8 Index = FirstIndexToDump;
-	uint8 RTIndex = 0;
-	while (Index != LastIndexToDump)
-	{
-		// this can wrap around at 256
-		uint8 NextIndex = Index + 1;
-		double Time = (GPUTimers[NextIndex] - GPUTimers[Index]) * 1000.0;
-		NSLog(@"Rende40rTarget %d took %.3fms", RTIndex, (float)Time);
-
-		Index = NextIndex;
-		RTIndex++;
-	}
-#endif
-
 	// xcode helper function
 	[CommandQueue insertDebugCaptureBoundary];
 
@@ -497,10 +452,6 @@ FMetalTexture2D* FMetalManager::GetBackBuffer()
 
 void FMetalManager::PrepareToDraw(uint32 NumVertices)
 {
-#if NO_MEMORY
-	return;
-#endif
-
 	SCOPE_CYCLE_COUNTER(STAT_MetalPrepareDrawTime);
 
 	NumDrawCalls++;
@@ -606,7 +557,7 @@ void FMetalManager::SetCurrentDepthStencilTarget(FMetalSurface* RenderSurface)
 {
 	if (RenderSurface)
 	{
-		// @todo urban: Handle stencil
+		// @todo metal stencil: track stencil here
 		CurrentDepthRenderTexture = RenderSurface->Texture;
 	}
 	else
@@ -644,7 +595,7 @@ void FMetalManager::UpdateContext()
 
 	MTLRenderPassAttachmentDescriptor* ColorAttachment = nil;
 	MTLRenderPassAttachmentDescriptor* DepthAttachment = nil;
-	// @todo urban: Handle stencil
+	// @todo metal stencil: handle stencil
 	MTLRenderPassAttachmentDescriptor* StencilAttachment = nil;
 
 	Pipeline.SampleCount = 0;
@@ -735,18 +686,10 @@ void FMetalManager::UpdateContext()
 		[CurrentCommandBuffer retain];
 		TRACK_OBJECT(CurrentCommandBuffer);
 
-#if GPU_TIMINGS
-		uint8 LocalTimerIndex = TimerIndex++;
-#endif
-
 		uint64 LocalCommandBufferIndex = CommandBufferIndex++;
 
 		[CurrentCommandBuffer addCompletedHandler : ^ (id <MTLCommandBuffer> Buffer)
 		{
-#if GPU_TIMINGS
-			// timer stuff
-			GPUTimers[LocalTimerIndex] = FPlatformTime::Seconds();
-#endif
 			FMetalManager::Get()->SetCompletedCommandBufferIndex(LocalCommandBufferIndex);
 		}];
 	}
@@ -842,9 +785,6 @@ inline int32 SetShaderResourcesFromBuffer(uint32 ShaderStage, FMetalUniformBuffe
 			const uint16 ResourceIndex = FRHIResourceTableEntry::GetResourceIndex(ResourceInfo);
 			const uint8 BindIndex = FRHIResourceTableEntry::GetBindIndex(ResourceInfo);
 
-			// @todo metal log: remove
-//			NSLog(@"Binding a resource at BindIndex %d, Resource Index %d", BindIndex, ResourceIndex);
-
 			// todo: could coalesce adjacent bound resources.
 			MetalResourceType* ResourcePtr = (MetalResourceType*)Buffer->RawResourceTable[ResourceIndex];
 			SetResource(ShaderStage, BindIndex, ResourcePtr);
@@ -862,9 +802,6 @@ void SetResourcesFromTables(ShaderType Shader, uint32 ShaderStage, uint32 Resour
 {
 	checkSlow(Shader);
 
-	// @todo metal log: remove
-//	NSLog(@"Resource Bits %x, DirtyBits: %llx", Shader->Bindings.ShaderResourceTable.ResourceTableBits, Shader->DirtyUniformBuffers);
-
 	// Mask the dirty bits by those buffers from which the shader has bound resources.
 	uint32 DirtyBits = Shader->Bindings.ShaderResourceTable.ResourceTableBits & Shader->DirtyUniformBuffers;
 	uint32 NumSetCalls = 0;
@@ -880,9 +817,6 @@ void SetResourcesFromTables(ShaderType Shader, uint32 ShaderStage, uint32 Resour
 		check(Buffer->GetLayout().GetHash() == Shader->Bindings.ShaderResourceTable.ResourceTableLayoutHashes[BufferIndex]);
 		Buffer->CacheResources(ResourceTableFrameCounter);
 
-		// @todo metal log: remove
-//		NSLog(@"Setting BufferIndex %d", BufferIndex);
-
 		// todo: could make this two pass: gather then set
 //		NumSetCalls += SetShaderResourcesFromBuffer<FMetalSurface>(ShaderStage, Buffer, Shader->Bindings.ShaderResourceTable.ShaderResourceViewMap.GetData(), BufferIndex);
 		NumSetCalls += SetShaderResourcesFromBuffer<FMetalSurface>(ShaderStage, Buffer, Shader->Bindings.ShaderResourceTable.TextureMap.GetData(), BufferIndex);
@@ -911,9 +845,6 @@ void FMetalManager::CommitGraphicsResourceTables()
 
 void FMetalManager::CommitNonComputeShaderConstants()
 {
-#if NO_MEMORY
-	return;
-#endif
 	ShaderParameters[METAL_SHADER_STAGE_VERTEX].CommitPackedUniformBuffers(CurrentBoundShaderState, METAL_SHADER_STAGE_VERTEX, CurrentBoundShaderState->VertexShader->BoundUniformBuffers, CurrentBoundShaderState->VertexShader->UniformBuffersCopyInfo);
 	ShaderParameters[METAL_SHADER_STAGE_VERTEX].CommitPackedGlobals(METAL_SHADER_STAGE_VERTEX, CurrentBoundShaderState->VertexShader->Bindings);
 	
@@ -976,15 +907,15 @@ void FMetalManager::SetRasterizerState(const FRasterizerStateInitializerRHI& Sta
 		ShadowRasterizerState.CullMode = State.CullMode;
 	}
 
-	// @todo urban: Shoud the clamp be FLT_MAX or 0 for "don't bother"
 	if (bFirstRasterizerState || ShadowRasterizerState.DepthBias != State.DepthBias || ShadowRasterizerState.SlopeScaleDepthBias != State.SlopeScaleDepthBias)
 	{
-		[CurrentContext setDepthBias:State.DepthBias slopeScale:State.SlopeScaleDepthBias clamp : FLT_MAX];
+		// no clamping
+		[CurrentContext setDepthBias:State.DepthBias slopeScale:State.SlopeScaleDepthBias clamp:FLT_MAX];
 		ShadowRasterizerState.DepthBias = State.DepthBias;
 		ShadowRasterizerState.SlopeScaleDepthBias = State.SlopeScaleDepthBias;
 	}
 
-	// @todo urban: Would we ever need this in a shipping app?
+	// @todo metal: Would we ever need this in a shipping app?
 #if !UE_BUILD_SHIPPING
 	if (bFirstRasterizerState || ShadowRasterizerState.FillMode != State.FillMode)
 	{
diff --git a/Engine/Source/Runtime/IOS/MetalRHI/Private/MetalManager.h b/Engine/Source/Runtime/IOS/MetalRHI/Private/MetalManager.h
--- a/Engine/Source/Runtime/IOS/MetalRHI/Private/MetalManager.h
+++ b/Engine/Source/Runtime/IOS/MetalRHI/Private/MetalManager.h
@@ -205,15 +205,6 @@ protected:
 	FRasterizerStateInitializerRHI ShadowRasterizerState;
 	bool bFirstRasterizerState;
 
-#if GPU_TIMINGS
-	// holds on to timing info
-	double GPUTimers[256];
-	uint8 TimerIndex;
-	uint8 FirstIndexThisFrame;
-	uint8 FirstIndexToDump;
-	uint8 LastIndexToDump;
-#endif
-
 	/** Apply the SRT before drawing */
 	void CommitGraphicsResourceTables();
 
diff --git a/Engine/Source/Runtime/IOS/MetalRHI/Private/MetalRHI.cpp b/Engine/Source/Runtime/IOS/MetalRHI/Private/MetalRHI.cpp
--- a/Engine/Source/Runtime/IOS/MetalRHI/Private/MetalRHI.cpp
+++ b/Engine/Source/Runtime/IOS/MetalRHI/Private/MetalRHI.cpp
@@ -48,7 +48,9 @@ FMetalDynamicRHI::FMetalDynamicRHI()
 	GMaxShadowDepthBufferSizeX = 4096;
 	GMaxShadowDepthBufferSizeY = 4096;
 // 	GReadTexturePoolSizeFromIni = true;
- 
+
+	GRHISupportsBaseVertexIndex = false;
+
 	GMaxTextureDimensions = 4096;
 	GMaxTextureMipCount = FPlatformMath::CeilLogTwo( GMaxTextureDimensions ) + 1;
 	GMaxTextureMipCount = FPlatformMath::Min<int32>( MAX_TEXTURE_MIP_COUNT, GMaxTextureMipCount );
diff --git a/Engine/Source/Runtime/IOS/MetalRHI/Private/MetalRHIPrivate.h b/Engine/Source/Runtime/IOS/MetalRHI/Private/MetalRHIPrivate.h
--- a/Engine/Source/Runtime/IOS/MetalRHI/Private/MetalRHIPrivate.h
+++ b/Engine/Source/Runtime/IOS/MetalRHI/Private/MetalRHIPrivate.h
@@ -9,7 +9,7 @@
 #include "Engine.h"
 
 // UE4 has a Max of 8 RTs, but we max out at 4, no need to waste time processing 8
-// @todo urban: Currently only using 1, so just do 1!
+// @todo metal mrt: Currently only using 1, so just do 1!
 const uint32 MaxMetalRenderTargets = 1;
 
 // How many possible vertex streams are allowed
@@ -21,8 +21,6 @@ const uint32 BufferOffsetAlignment = 256;
 //#define BUFFER_CACHE_MODE MTLResourceOptionCPUCacheModeDefault
 #define BUFFER_CACHE_MODE MTLResourceOptionCPUCacheModeWriteCombined
 
-#define GPU_TIMINGS 0 // (UE_BUILD_DEBUG || UE_BUILD_DEVELOPMENT)
-
 #define SHOULD_TRACK_OBJECTS 0 // (UE_BUILD_DEBUG || UE_BUILD_DEVELOPMENT)
 
 #define UNREAL_TO_METAL_BUFFER_INDEX(Index) (15 - Index)
diff --git a/Engine/Source/Runtime/IOS/MetalRHI/Private/MetalUniformBuffer.cpp b/Engine/Source/Runtime/IOS/MetalRHI/Private/MetalUniformBuffer.cpp
--- a/Engine/Source/Runtime/IOS/MetalRHI/Private/MetalUniformBuffer.cpp
+++ b/Engine/Source/Runtime/IOS/MetalRHI/Private/MetalUniformBuffer.cpp
@@ -220,7 +220,7 @@ FMetalUniformBuffer::FMetalUniformBuffer(const void* Contents, const FRHIUniform
 		FMemory::Memcpy(((uint8*)[Buffer contents]) + Offset, Contents, Layout.ConstantBufferSize);
 	}
 
-	// @todo metal: Add support for resource tables!
+	// set up an SRT-style uniform buffer
 	if (Layout.Resources.Num())
 	{
 		int32 NumResources = Layout.Resources.Num();
diff --git a/Engine/Source/Runtime/IOS/MetalRHI/Private/MetalVertexDeclaration.cpp b/Engine/Source/Runtime/IOS/MetalRHI/Private/MetalVertexDeclaration.cpp
--- a/Engine/Source/Runtime/IOS/MetalRHI/Private/MetalVertexDeclaration.cpp
+++ b/Engine/Source/Runtime/IOS/MetalRHI/Private/MetalVertexDeclaration.cpp
@@ -90,12 +90,8 @@ void FMetalVertexDeclaration::GenerateLayout(const FVertexDeclarationElementList
 	{
 		const FVertexElement& Element = Elements[ElementIndex];
 		
-		// @todo urban: for zero stride elements, assume a repeated color
-		// @todo urban: Make sure zero stride actually works now!! If so, then we can remove the hack zero stride buffer
-		uint32 Stride = Element.Stride;// ? Element.Stride : 4;
-
-		checkf(Stride == 0 || Element.Offset + TranslateElementTypeToSize(Element.Type) <= Stride, 
-			TEXT("Stream component is bigger than stride: Offset: %d, Size: %d [Type %d], Stride: %d"), Element.Offset, TranslateElementTypeToSize(Element.Type), (uint32)Element.Type, Stride);
+		checkf(Element.Stride == 0 || Element.Offset + TranslateElementTypeToSize(Element.Type) <= Element.Stride, 
+			TEXT("Stream component is bigger than stride: Offset: %d, Size: %d [Type %d], Stride: %d"), Element.Offset, TranslateElementTypeToSize(Element.Type), (uint32)Element.Type, Element.Stride);
 
 		// we offset 6 buffers to leave space for uniform buffers
 		uint32 ShaderBufferIndex = UNREAL_TO_METAL_BUFFER_INDEX(Element.StreamIndex);
@@ -104,29 +100,32 @@ void FMetalVertexDeclaration::GenerateLayout(const FVertexDeclarationElementList
 		uint32* ExistingStride = BufferStrides.Find(ShaderBufferIndex);
 		if (ExistingStride == NULL)
 		{
-			BufferStrides.Add(ShaderBufferIndex, Stride);
+			// handle 0 stride buffers
+			MTLVertexStepFunction Function = (Element.Stride == 0 ? MTLVertexStepFunctionConstant : MTLVertexStepFunctionPerVertex);
+			uint32 StepRate = (Element.Stride == 0 ? 0 : 1);
+			// even with MTLVertexStepFunctionConstant, it needs a non-zero stride (not sure why)
+			uint32 Stride = (Element.Stride == 0 ? 4 : Element.Stride);
+
+			// look for any unset strides coming from UE4 (this can be removed when all are fixed)
+			if (Element.Stride == 0xFFFF)
+			{
+				NSLog(@"Setting illegal stride - break here if you want to find out why, but this won't break until we try to render with it");
+				Stride = 200;
+			}
+
+			// set the stride once per buffer
+			[Layout setStride:Stride stepFunction:Function stepRate:StepRate atVertexBufferIndex:ShaderBufferIndex];
+
+			// track this buffer and stride
+			BufferStrides.Add(ShaderBufferIndex, Element.Stride);
 		}
 		else
 		{
-			check(Stride == *ExistingStride);
+			// if the strides of elements with same buffer index have different strides, something is VERY wrong
+			check(Element.Stride == *ExistingStride);
 		}
 
+		// set the format for each element
 		[Layout setVertexFormat:TranslateElementTypeToMTLType(Element.Type) offset:Element.Offset vertexBufferIndex:ShaderBufferIndex atAttributeIndex:Element.AttributeIndex];
 	}
-
-	for (auto It = BufferStrides.CreateIterator(); It; ++It)
-	{
-		uint32 Stride = It.Value();
-		if (Stride == 0xFFFF)
-		{
-			NSLog(@"Setting illegal stride - break here if you want to find out why, but this won't break until we try to render with it");
-			Stride = 200;// 16;
-		}
-
-		// handle 0 stride buffers
-		MTLVertexStepFunction Function = (Stride == 0 ? MTLVertexStepFunctionConstant : MTLVertexStepFunctionPerVertex);
-		uint32 StepRate = (Stride == 0 ? 0 : 1);
-		uint32 UsedStride = (Stride == 0 ? 4 : Stride);
-		[Layout setStride:UsedStride stepFunction:Function stepRate:StepRate atVertexBufferIndex:It.Key()];
-	}
 }
diff --git a/Engine/Source/Runtime/RHI/Private/RHI.cpp b/Engine/Source/Runtime/RHI/Private/RHI.cpp
--- a/Engine/Source/Runtime/RHI/Private/RHI.cpp
+++ b/Engine/Source/Runtime/RHI/Private/RHI.cpp
@@ -140,6 +140,7 @@ bool GVertexElementsCanShareStreamOffset = true;
 bool GTriggerGPUProfile = false;
 bool GRHISupportsTextureStreaming = false;
 bool GSupportsDepthBoundsTest = false;
+bool GRHISupportsBaseVertexIndex = true;
 
 /** Whether we are profiling GPU hitches. */
 bool GTriggerGPUHitchProfile = false;
diff --git a/Engine/Source/Runtime/RHI/Public/RHI.h b/Engine/Source/Runtime/RHI/Public/RHI.h
--- a/Engine/Source/Runtime/RHI/Public/RHI.h
+++ b/Engine/Source/Runtime/RHI/Public/RHI.h
@@ -196,6 +196,10 @@ extern RHI_API int32 GPoolSizeVRAMPercentage;
 extern RHI_API int32 GNumDrawCallsRHI;
 extern RHI_API int32 GNumPrimitivesDrawnRHI;
 
+/** Whether or not the RHI can handle a non-zero BaseVertexIndex - extra SetStreamSource calls will be needed if this is false */
+extern RHI_API bool GRHISupportsBaseVertexIndex;
+
+
 /** Called once per frame only from within an RHI. */
 extern RHI_API void RHIPrivateBeginFrame();
 
diff --git a/Engine/Source/Runtime/SlateRHIRenderer/Private/SlateRHIRenderingPolicy.cpp b/Engine/Source/Runtime/SlateRHIRenderer/Private/SlateRHIRenderingPolicy.cpp
--- a/Engine/Source/Runtime/SlateRHIRenderer/Private/SlateRHIRenderingPolicy.cpp
+++ b/Engine/Source/Runtime/SlateRHIRenderer/Private/SlateRHIRenderingPolicy.cpp
@@ -303,7 +303,10 @@ void FSlateRHIRenderingPolicy::DrawElements(FRHICommandListImmediate& RHICmdList
 
 	FSamplerStateRHIRef BilinearClamp = TStaticSamplerState<SF_Bilinear,AM_Clamp,AM_Clamp,AM_Clamp>::GetRHI();
 
-	RHICmdList.SetStreamSource(0, VertexBuffer.VertexBufferRHI, sizeof(FSlateVertex), 0);
+	if (GRHISupportsBaseVertexIndex)
+	{
+		RHICmdList.SetStreamSource(0, VertexBuffer.VertexBufferRHI, sizeof(FSlateVertex), 0);
+	}
 
 	// Draw each element
 	for( int32 BatchIndex = 0; BatchIndex < RenderBatches.Num(); ++BatchIndex )
@@ -419,7 +422,16 @@ void FSlateRHIRenderingPolicy::DrawElements(FRHICommandListImmediate& RHICmdList
 
 				uint32 PrimitiveCount = RenderBatch.DrawPrimitiveType == ESlateDrawPrimitive::LineList ? RenderBatch.NumIndices / 2 : RenderBatch.NumIndices / 3;
 
-				RHICmdList.DrawIndexedPrimitive(IndexBuffer.IndexBufferRHI, GetRHIPrimitiveType(RenderBatch.DrawPrimitiveType), RenderBatch.VertexOffset, 0, RenderBatch.NumVertices, RenderBatch.IndexOffset, PrimitiveCount, 1);
+				// for RHIs that can't handle VertexOffset, we need to offset the stream source each time
+				if (!GRHISupportsBaseVertexIndex)
+				{
+					RHICmdList.SetStreamSource(0, VertexBuffer.VertexBufferRHI, sizeof(FSlateVertex), RenderBatch.VertexOffset * sizeof(FSlateVertex));
+					RHICmdList.DrawIndexedPrimitive(IndexBuffer.IndexBufferRHI, GetRHIPrimitiveType(RenderBatch.DrawPrimitiveType), 0, 0, RenderBatch.NumVertices, RenderBatch.IndexOffset, PrimitiveCount, 1);
+				}
+				else
+				{
+					RHICmdList.DrawIndexedPrimitive(IndexBuffer.IndexBufferRHI, GetRHIPrimitiveType(RenderBatch.DrawPrimitiveType), RenderBatch.VertexOffset, 0, RenderBatch.NumVertices, RenderBatch.IndexOffset, PrimitiveCount, 1);
+				}
 
 			}
 			else if (ShaderResource && ShaderResource->GetType() == ESlateShaderResource::Material)
@@ -457,7 +469,16 @@ void FSlateRHIRenderingPolicy::DrawElements(FRHICommandListImmediate& RHICmdList
 
 					uint32 PrimitiveCount = RenderBatch.DrawPrimitiveType == ESlateDrawPrimitive::LineList ? RenderBatch.NumIndices / 2 : RenderBatch.NumIndices / 3;
 
-					RHICmdList.DrawIndexedPrimitive(IndexBuffer.IndexBufferRHI, GetRHIPrimitiveType(RenderBatch.DrawPrimitiveType), RenderBatch.VertexOffset, 0, RenderBatch.NumVertices, RenderBatch.IndexOffset, PrimitiveCount, 1);
+					// for RHIs that can't handle VertexOffset, we need to offset the stream source each time
+					if (!GRHISupportsBaseVertexIndex)
+					{
+						RHICmdList.SetStreamSource(0, VertexBuffer.VertexBufferRHI, sizeof(FSlateVertex), RenderBatch.VertexOffset * sizeof(FSlateVertex));
+						RHICmdList.DrawIndexedPrimitive(IndexBuffer.IndexBufferRHI, GetRHIPrimitiveType(RenderBatch.DrawPrimitiveType), 0, 0, RenderBatch.NumVertices, RenderBatch.IndexOffset, PrimitiveCount, 1);
+					}
+					else
+					{
+						RHICmdList.DrawIndexedPrimitive(IndexBuffer.IndexBufferRHI, GetRHIPrimitiveType(RenderBatch.DrawPrimitiveType), RenderBatch.VertexOffset, 0, RenderBatch.NumVertices, RenderBatch.IndexOffset, PrimitiveCount, 1);
+					}
 				}
 			}
 
