Metal beta 3 support integrated to main branch

diff --git a/Engine/Source/Runtime/IOS/MetalRHI/Private/MetalManager.cpp b/Engine/Source/Runtime/IOS/MetalRHI/Private/MetalManager.cpp
--- a/Engine/Source/Runtime/IOS/MetalRHI/Private/MetalManager.cpp
+++ b/Engine/Source/Runtime/IOS/MetalRHI/Private/MetalManager.cpp
@@ -60,7 +60,7 @@ void FPipelineShadow::SetHash(uint64 InHash)
 
 	bIsDepthWriteEnabled = GET_HASH(OFFSET_DEPTH_WRITE_ENABLED, NUMBITS_DEPTH_WRITE_ENABLED);
 	bIsStencilWriteEnabled = GET_HASH(OFFSET_STENCIL_WRITE_ENABLED, NUMBITS_STENCIL_WRITE_ENABLED);
-	RenderTargets[0] = [[MTLRenderPipelineAttachmentDescriptor alloc] init];
+	RenderTargets[0] = [[MTLRenderPipelineColorAttachmentDescriptor alloc] init];
 	FMetalManager::Get()->ReleaseObject(RenderTargets[0]);
 	RenderTargets[0].sourceRGBBlendFactor = (MTLBlendFactor)GET_HASH(OFFSET_SOURCE_RGB_BLEND_FACTOR, NUMBITS_SOURCE_RGB_BLEND_FACTOR);
 	RenderTargets[0].destinationRGBBlendFactor = (MTLBlendFactor)GET_HASH(OFFSET_DEST_RGB_BLEND_FACTOR, NUMBITS_DEST_RGB_BLEND_FACTOR);
@@ -100,48 +100,6 @@ static MTLCullMode TranslateCullMode(ERasterizerCullMode CullMode)
 	}
 }
 
-static MTLBlendOperation TranslateBlendOp(EBlendOperation BlendOp)
-{
-	switch (BlendOp)
-	{
-	case BO_Subtract:	return MTLBlendOperationSubtract;
-	case BO_Min:		return MTLBlendOperationMin;
-	case BO_Max:		return MTLBlendOperationMax;
-	default:			return MTLBlendOperationAdd;
-	};
-}
-
-
-static MTLBlendFactor TranslateBlendFactor(EBlendFactor BlendFactor)
-{
-	switch (BlendFactor)
-	{
-	case BF_One:					return MTLBlendFactorOne;
-	case BF_SourceColor:			return MTLBlendFactorSourceColor;
-	case BF_InverseSourceColor:		return MTLBlendFactorOneMinusSourceColor;
-	case BF_SourceAlpha:			return MTLBlendFactorSourceAlpha;
-	case BF_InverseSourceAlpha:		return MTLBlendFactorOneMinusSourceAlpha;
-	case BF_DestAlpha:				return MTLBlendFactorDestinationAlpha;
-	case BF_InverseDestAlpha:		return MTLBlendFactorOneMinusDestinationAlpha;
-	case BF_DestColor:				return MTLBlendFactorDestinationColor;
-	case BF_InverseDestColor:		return MTLBlendFactorOneMinusDestinationColor;
-	default:						return MTLBlendFactorZero;
-	};
-}
-
-static MTLColorWriteMask TranslateWriteMask(EColorWriteMask WriteMask)
-{
-	uint32 Result = 0;
-	Result |= (WriteMask & CW_RED) ? (MTLColorWriteMaskRed) : 0;
-	Result |= (WriteMask & CW_GREEN) ? (MTLColorWriteMaskGreen) : 0;
-	Result |= (WriteMask & CW_BLUE) ? (MTLColorWriteMaskBlue) : 0;
-	Result |= (WriteMask & CW_ALPHA) ? (MTLColorWriteMaskAlpha) : 0;
-
-	return (MTLColorWriteMask)Result;
-}
-
-
-
 
 id<MTLRenderPipelineState> FPipelineShadow::CreatePipelineStateForBoundShaderState(FMetalBoundShaderState* BSS) const
 {
@@ -250,7 +208,6 @@ FMetalManager::FMetalManager()
 	: Device([IOSAppDelegate GetDelegate].IOSView->MetalDevice)
 	, CurrentCommandBuffer(nil)
 	, CurrentDrawable(nil)
-	, CurrentRenderPass(nil)
 	, CurrentContext(nil)
 	, CurrentColorRenderTexture(nil)
 	, PreviousColorRenderTexture(nil)
@@ -289,6 +246,20 @@ FMetalManager::FMetalManager()
 	InitFrame();
 }
 
+void FMetalManager::CreateAutoreleasePool()
+{
+	if (FPlatformTLS::GetTlsValue(AutoReleasePoolTLSSlot) == NULL)
+	{
+		FPlatformTLS::SetTlsValue(AutoReleasePoolTLSSlot, FPlatformMisc::CreateAutoreleasePool());
+	}
+}
+
+void FMetalManager::DrainAutoreleasePool()
+{
+	FPlatformMisc::ReleaseAutoreleasePool(FPlatformTLS::GetTlsValue(AutoReleasePoolTLSSlot));
+	FPlatformTLS::SetTlsValue(AutoReleasePoolTLSSlot, NULL);
+}
+
 void FMetalManager::BeginScene()
 {
 	// Increment the frame counter. INDEX_NONE is a special value meaning "uninitialized", so if
@@ -316,15 +287,10 @@ void FMetalManager::BeginFrame()
 }
 
 
-static TConsoleVariableData<int32>* GDrawableMode = IConsoleManager::Get().RegisterConsoleVariable(TEXT("u.Drawable"), 1, TEXT("Whether to use new drawable mode or not"), ECVF_Default)->AsVariableInt();
-static bool GUseNewDrawable = true;
-
 void FMetalManager::InitFrame()
 {
-	GUseNewDrawable = GDrawableMode->GetValueOnRenderThread() > 0 ? true : false;
-
 	// start an auto release pool (EndFrame will drain and remake)
-	FPlatformTLS::SetTlsValue(AutoReleasePoolTLSSlot, FPlatformMisc::CreateAutoreleasePool());
+	CreateAutoreleasePool();
 
 	// create the command buffer for this frame
 	CurrentCommandBuffer = [CommandQueue commandBufferWithUnretainedReferences];
@@ -343,20 +309,8 @@ void FMetalManager::InitFrame()
 
 //	double Mid = FPlatformTime::Seconds();
 
-	if (GUseNewDrawable)
-	{
-		// mark us to get later
-		BackBuffer->Surface.Texture = nil;
-	}
-	else
-	{
-		SCOPE_CYCLE_COUNTER(STAT_MetalMakeDrawableTime);
-
-		// make a drawable object for this frame
-		CurrentDrawable = [[IOSAppDelegate GetDelegate].IOSView MakeDrawable];
-		TRACK_OBJECT(CurrentDrawable);
-		BackBuffer->Surface.Texture = CurrentDrawable.texture;
-	}
+	// mark us to get later
+	BackBuffer->Surface.Texture = nil;
 
 //	double End = FPlatformTime::Seconds();
 //	NSLog(@"Semaphore Block Time: %.2f   -- MakeDrawable Time: %.2f", 1000.0f*(Mid - Start), 1000.0f*(End - Mid));
@@ -373,7 +327,7 @@ void FMetalManager::EndFrame(bool bPresent)
 //	NSLog(@"There were %d draw calls for final RT in frame %lld", NumDrawCalls, GFrameCounter);
 	// commit the render context to the commandBuffer
 	[CurrentContext endEncoding];
-	ReleaseObject(CurrentContext);
+	[CurrentContext release];
 	CurrentContext = nil;
 
 	// kick the whole buffer
@@ -398,26 +352,16 @@ void FMetalManager::EndFrame(bool bPresent)
 			[CurrentDrawable present];
 		}
 
-		UNTRACK_OBJECT(CurrentDrawable);
-		[CurrentDrawable release];
 		CurrentDrawable = nil;
 	}
 
-
-	ReleaseObject(CurrentRenderPass);
-	CurrentRenderPass = nil;
-
-
 	UNTRACK_OBJECT(CurrentCommandBuffer);
 	[CurrentCommandBuffer release];
 
 	// xcode helper function
 	[CommandQueue insertDebugCaptureBoundary];
 
-	// drain the pool
-	FPlatformMisc::ReleaseAutoreleasePool(FPlatformTLS::GetTlsValue(AutoReleasePoolTLSSlot));
-
-	// drain the oldest pool
+	// drain the oldest delayed free list
 	uint32 PrevFreeList = (WhichFreeList + 1) % ARRAY_COUNT(DelayedFreeLists);
 	for (int32 Index = 0; Index < DelayedFreeLists[PrevFreeList].Num(); Index++)
 	{
@@ -441,6 +385,9 @@ void FMetalManager::EndFrame(bool bPresent)
 	}
 #endif
 
+	// drain the pool
+	DrainAutoreleasePool();
+
 	InitFrame();
 }
 
@@ -477,36 +424,29 @@ void FMetalManager::SetBlendState(FMetalBlendState* BlendState)
 	for(uint32 RenderTargetIndex = 0;RenderTargetIndex < MaxMetalRenderTargets; ++RenderTargetIndex)
 	{
 		// @todo metal mrt
+		MTLRenderPipelineColorAttachmentDescriptor* Blend0 = BlendState->RenderTargetStates[RenderTargetIndex];
+		MTLRenderPipelineColorAttachmentDescriptor* Dest = Pipeline.RenderTargets[RenderTargetIndex];
 		check(RenderTargetIndex == 0);
 
-		// get the blend params
-		const FBlendStateInitializerRHI::FRenderTarget& Init =
-			BlendState->RenderTargetStates.bUseIndependentRenderTargetBlendStates
-				? BlendState->RenderTargetStates.RenderTargets[RenderTargetIndex]
-				: BlendState->RenderTargetStates.RenderTargets[0];
+		// assign the struct by value
+//		*Pipeline.RenderTargets[RenderTargetIndex] = *Blend0;
+		Dest.blendingEnabled = Blend0.blendingEnabled;
+		Dest.sourceRGBBlendFactor = Blend0.sourceRGBBlendFactor;
+		Dest.destinationRGBBlendFactor = Blend0.destinationRGBBlendFactor;
+		Dest.rgbBlendOperation = Blend0.rgbBlendOperation;
+		Dest.sourceAlphaBlendFactor = Blend0.sourceAlphaBlendFactor;
+		Dest.destinationAlphaBlendFactor = Blend0.destinationAlphaBlendFactor;
+		Dest.alphaBlendOperation = Blend0.alphaBlendOperation;
+		Dest.writeMask = Blend0.writeMask;
 
-		// write into the pipline attachment
-		MTLRenderPipelineAttachmentDescriptor* Attachment = Pipeline.RenderTargets[RenderTargetIndex];
-		
-		// set values
-		Attachment.blendingEnabled =
-			Init.ColorBlendOp != BO_Add || Init.ColorDestBlend != BF_Zero || Init.ColorSrcBlend != BF_One ||
-			Init.AlphaBlendOp != BO_Add || Init.AlphaDestBlend != BF_Zero || Init.AlphaSrcBlend != BF_One;
-		Attachment.sourceRGBBlendFactor = TranslateBlendFactor(Init.ColorSrcBlend);
-		Attachment.destinationRGBBlendFactor = TranslateBlendFactor(Init.ColorDestBlend);
-		Attachment.rgbBlendOperation = TranslateBlendOp(Init.ColorBlendOp);
-		Attachment.sourceAlphaBlendFactor = TranslateBlendFactor(Init.AlphaSrcBlend);
-		Attachment.destinationAlphaBlendFactor = TranslateBlendFactor(Init.AlphaDestBlend);
-		Attachment.alphaBlendOperation = TranslateBlendOp(Init.AlphaBlendOp);
-		Attachment.writeMask = TranslateWriteMask(Init.ColorWriteMask);
-
-		SET_HASH(OFFSET_SOURCE_RGB_BLEND_FACTOR, NUMBITS_SOURCE_RGB_BLEND_FACTOR, Attachment.sourceRGBBlendFactor);
-		SET_HASH(OFFSET_DEST_RGB_BLEND_FACTOR, NUMBITS_DEST_RGB_BLEND_FACTOR, Attachment.destinationRGBBlendFactor);
-		SET_HASH(OFFSET_RGB_BLEND_OPERATION, NUMBITS_RGB_BLEND_OPERATION, Attachment.rgbBlendOperation);
-		SET_HASH(OFFSET_SOURCE_A_BLEND_FACTOR, NUMBITS_SOURCE_A_BLEND_FACTOR, Attachment.sourceAlphaBlendFactor);
-		SET_HASH(OFFSET_DEST_A_BLEND_FACTOR, NUMBITS_DEST_A_BLEND_FACTOR, Attachment.destinationAlphaBlendFactor);
-		SET_HASH(OFFSET_A_BLEND_OPERATION, NUMBITS_A_BLEND_OPERATION, Attachment.alphaBlendOperation);
-		SET_HASH(OFFSET_WRITE_MASK, NUMBITS_WRITE_MASK, Attachment.writeMask);
+
+		SET_HASH(OFFSET_SOURCE_RGB_BLEND_FACTOR, NUMBITS_SOURCE_RGB_BLEND_FACTOR, Blend0.sourceRGBBlendFactor);
+		SET_HASH(OFFSET_DEST_RGB_BLEND_FACTOR, NUMBITS_DEST_RGB_BLEND_FACTOR, Blend0.destinationRGBBlendFactor);
+		SET_HASH(OFFSET_RGB_BLEND_OPERATION, NUMBITS_RGB_BLEND_OPERATION, Blend0.rgbBlendOperation);
+		SET_HASH(OFFSET_SOURCE_A_BLEND_FACTOR, NUMBITS_SOURCE_A_BLEND_FACTOR, Blend0.sourceAlphaBlendFactor);
+		SET_HASH(OFFSET_DEST_A_BLEND_FACTOR, NUMBITS_DEST_A_BLEND_FACTOR, Blend0.destinationAlphaBlendFactor);
+		SET_HASH(OFFSET_A_BLEND_OPERATION, NUMBITS_A_BLEND_OPERATION, Blend0.alphaBlendOperation);
+		SET_HASH(OFFSET_WRITE_MASK, NUMBITS_WRITE_MASK, Blend0.writeMask);
 	}
 }
 
@@ -523,25 +463,21 @@ void FMetalManager::SetCurrentRenderTarget(FMetalSurface* RenderSurface)
 	// update the current rendered-to pixel format
 	if (RenderSurface)
 	{
-		if (GUseNewDrawable)
+		// first time in a frame that we are setting the backbuffer, get it
+		if (RenderSurface == &BackBuffer->Surface && RenderSurface->Texture == nil && CurrentDrawable == nil)
 		{
-			// first time in a frame that we are setting the backbuffer, get it
-			if (RenderSurface == &BackBuffer->Surface && RenderSurface->Texture == nil && CurrentDrawable == nil)
-			{
-				SCOPE_CYCLE_COUNTER(STAT_MetalMakeDrawableTime);
+			SCOPE_CYCLE_COUNTER(STAT_MetalMakeDrawableTime);
 
-				uint32 IdleStart = FPlatformTime::Cycles();
+			uint32 IdleStart = FPlatformTime::Cycles();
 
 				// make a drawable object for this frame
 				CurrentDrawable = [[IOSAppDelegate GetDelegate].IOSView MakeDrawable];
-				TRACK_OBJECT(CurrentDrawable);
 
-				GRenderThreadIdle[ERenderThreadIdleTypes::WaitingForGPUPresent] += FPlatformTime::Cycles() - IdleStart;
-				GRenderThreadNumIdle[ERenderThreadIdleTypes::WaitingForGPUPresent]++;
+			GRenderThreadIdle[ERenderThreadIdleTypes::WaitingForGPUPresent] += FPlatformTime::Cycles() - IdleStart;
+			GRenderThreadNumIdle[ERenderThreadIdleTypes::WaitingForGPUPresent]++;
 
-				// set the texture into the backbuffer
-				RenderSurface->Texture = CurrentDrawable.texture;
-			}
+			// set the texture into the backbuffer
+			RenderSurface->Texture = CurrentDrawable.texture;
 		}
 
 		CurrentColorRenderTexture = RenderSurface->Texture;
@@ -580,28 +516,24 @@ void FMetalManager::UpdateContext()
 		return;
 	}
 
-	// toss previous render pass object
-	ReleaseObject(CurrentRenderPass);
+	// make a new one (autoreleased)
+	MTLRenderPassDescriptor* CurrentRenderPass = [MTLRenderPassDescriptor renderPassDescriptor];
 
-	// make a new one
-	CurrentRenderPass = [MTLRenderPassDescriptor renderPassDescriptor];
-	[CurrentRenderPass retain];
-	TRACK_OBJECT(CurrentRenderPass);
 	// if we need to do queries, write to the ring buffer (we set the offset into the ring buffer per query)
 	CurrentRenderPass.visibilityResultBuffer = QueryBuffer.Buffer;
 
 	// @todo metal mrt: Set an attachment index properly
 	uint32 AttachmentIndex = 0;
 
-	MTLRenderPassAttachmentDescriptor* ColorAttachment = nil;
-	MTLRenderPassAttachmentDescriptor* DepthAttachment = nil;
+	MTLRenderPassColorAttachmentDescriptor* ColorAttachment = nil;
+	MTLRenderPassDepthAttachmentDescriptor* DepthAttachment = nil;
 	// @todo metal stencil: handle stencil
 	MTLRenderPassAttachmentDescriptor* StencilAttachment = nil;
 
 	Pipeline.SampleCount = 0;
 	if (CurrentColorRenderTexture)
 	{
-		ColorAttachment = [[MTLRenderPassAttachmentDescriptor alloc] init];
+		ColorAttachment = [[MTLRenderPassColorAttachmentDescriptor alloc] init];
 
 		if (CurrentMSAARenderTexture)
 		{
@@ -632,12 +564,12 @@ void FMetalManager::UpdateContext()
 	
 	if (CurrentDepthRenderTexture)
 	{
-		DepthAttachment = [[MTLRenderPassAttachmentDescriptor alloc] init];
+		DepthAttachment = [[MTLRenderPassDepthAttachmentDescriptor alloc] init];
 
 		DepthAttachment.texture = CurrentDepthRenderTexture;
 		[DepthAttachment setLoadAction:MTLLoadActionClear];
 		[DepthAttachment setStoreAction:MTLStoreActionDontCare];
-		[DepthAttachment setClearValue:MTLClearValueMakeDepth(0.0f)];
+		[DepthAttachment setClearDepth:0.0];
 
 		Pipeline.DepthTargetFormat = CurrentDepthRenderTexture.pixelFormat;
 		if (Pipeline.SampleCount == 0)
@@ -666,13 +598,10 @@ void FMetalManager::UpdateContext()
 			NSLog(@"There were %d draw calls for an RT in frame %lld", NumDrawCalls, GFrameCounter);
 		}
 
-		if (true)//NumDrawCalls > 0)
-		{
-			[CurrentContext endEncoding];
-			NumDrawCalls = 0;
-		}
+		[CurrentContext endEncoding];
+		NumDrawCalls = 0;
 
-		ReleaseObject(CurrentContext);
+		[CurrentContext release];
 
 		// if we are doing occlusion queries, we could use this method, along with a completion callback
 		// to set a "render target complete" flag that the OQ code next frame would wait on
@@ -925,4 +854,4 @@ void FMetalManager::SetRasterizerState(const FRasterizerStateInitializerRHI& Sta
 #endif
 
 	bFirstRasterizerState = false;
-}
+}
\ No newline at end of file
