- More files for Metal beta3 into main branch
- Somehow these files went in to default changelist, instead of being submitted. What?!?!?!?

diff --git a/Engine/Source/Runtime/IOS/MetalRHI/Private/MetalManager.h b/Engine/Source/Runtime/IOS/MetalRHI/Private/MetalManager.h
--- a/Engine/Source/Runtime/IOS/MetalRHI/Private/MetalManager.h
+++ b/Engine/Source/Runtime/IOS/MetalRHI/Private/MetalManager.h
@@ -17,7 +17,7 @@ struct FPipelineShadow
 
 		for (int Index = 0; Index < MaxMetalRenderTargets; Index++)
 		{
-			RenderTargets[Index] = [[MTLRenderPipelineAttachmentDescriptor alloc] init];
+			RenderTargets[Index] = [[MTLRenderPipelineColorAttachmentDescriptor alloc] init];
 		}
 	}
 
@@ -41,7 +41,7 @@ struct FPipelineShadow
 
 	bool bIsDepthWriteEnabled;
 	bool bIsStencilWriteEnabled;
-	MTLRenderPipelineAttachmentDescriptor* RenderTargets[MaxMetalRenderTargets];
+	MTLRenderPipelineColorAttachmentDescriptor* RenderTargets[MaxMetalRenderTargets];
 	MTLPixelFormat DepthTargetFormat;
     uint32 SampleCount;
 
@@ -144,6 +144,12 @@ public:
 		return ShaderParameters[Stage];
 	}
 
+	/**
+	 * Handle rendering thread starting/stopping
+	 */
+	void CreateAutoreleasePool();
+	void DrainAutoreleasePool();
+
 protected:
 	FMetalManager();
 	void InitFrame();
@@ -159,7 +165,6 @@ protected:
 
 	dispatch_semaphore_t CommandBufferSemaphore;
 
-	MTLRenderPassDescriptor* CurrentRenderPass;
 	id<MTLRenderCommandEncoder> CurrentContext;
 	
 	id<MTLTexture> CurrentColorRenderTexture, PreviousColorRenderTexture;
diff --git a/Engine/Source/Runtime/IOS/MetalRHI/Private/MetalRHI.cpp b/Engine/Source/Runtime/IOS/MetalRHI/Private/MetalRHI.cpp
--- a/Engine/Source/Runtime/IOS/MetalRHI/Private/MetalRHI.cpp
+++ b/Engine/Source/Runtime/IOS/MetalRHI/Private/MetalRHI.cpp
@@ -92,7 +92,7 @@ FMetalDynamicRHI::FMetalDynamicRHI()
 	GPixelFormats[PF_G16R16F			].PlatformFormat	= MTLPixelFormatRG16Float;
 	GPixelFormats[PF_G16R16F_FILTER		].PlatformFormat	= MTLPixelFormatRG16Float;
 	GPixelFormats[PF_G32R32F			].PlatformFormat	= MTLPixelFormatRG32Float;
-	GPixelFormats[PF_A2B10G10R10		].PlatformFormat    = MTLPixelFormatA2BGR10Unorm;
+	GPixelFormats[PF_A2B10G10R10		].PlatformFormat    = MTLPixelFormatRGB10A2Unorm;
 	GPixelFormats[PF_A16B16G16R16		].PlatformFormat    = MTLPixelFormatRGBA16Float;
 	GPixelFormats[PF_D24				].PlatformFormat	= MTLPixelFormatInvalid;
 	GPixelFormats[PF_R16F				].PlatformFormat	= MTLPixelFormatR16Float;
@@ -100,14 +100,14 @@ FMetalDynamicRHI::FMetalDynamicRHI()
 	GPixelFormats[PF_BC5				].PlatformFormat	= MTLPixelFormatInvalid;
 	GPixelFormats[PF_V8U8				].PlatformFormat	= 
 	GPixelFormats[PF_A1					].PlatformFormat	= MTLPixelFormatInvalid;
-	GPixelFormats[PF_FloatR11G11B10		].PlatformFormat	= MTLPixelFormatB10GR11Float;
+	GPixelFormats[PF_FloatR11G11B10		].PlatformFormat	= MTLPixelFormatRG11B10Float;
 	GPixelFormats[PF_FloatR11G11B10		].BlockBytes		= 4;
 	GPixelFormats[PF_A8					].PlatformFormat	= MTLPixelFormatR8Unorm;
 	GPixelFormats[PF_R32_UINT			].PlatformFormat	= MTLPixelFormatR32Uint;
 	GPixelFormats[PF_R32_SINT			].PlatformFormat	= MTLPixelFormatR32Sint;
 	GPixelFormats[PF_R16G16B16A16_UINT	].PlatformFormat	= MTLPixelFormatRGBA16Uint;
 	GPixelFormats[PF_R16G16B16A16_SINT	].PlatformFormat	= MTLPixelFormatRGBA16Sint;
-	GPixelFormats[PF_R5G6B5_UNORM		].PlatformFormat	= MTLPixelFormatRGB565Unorm;
+	GPixelFormats[PF_R5G6B5_UNORM		].PlatformFormat	= MTLPixelFormatB5G6R5Unorm;
 	GPixelFormats[PF_R8G8B8A8			].PlatformFormat	= MTLPixelFormatRGBA8Unorm;
 	GPixelFormats[PF_R8G8				].PlatformFormat	= MTLPixelFormatRG8Unorm;
 
@@ -205,11 +205,12 @@ void FMetalDynamicRHI::RHIFlushResources()
 
 void FMetalDynamicRHI::RHIAcquireThreadOwnership()
 {
-
+	FMetalManager::Get()->CreateAutoreleasePool();
 }
+
 void FMetalDynamicRHI::RHIReleaseThreadOwnership()
 {
-
+	FMetalManager::Get()->DrainAutoreleasePool();
 }
 
 void FMetalDynamicRHI::RHIGpuTimeBegin(uint32 Hash, bool bCompute)
diff --git a/Engine/Source/Runtime/IOS/MetalRHI/Private/MetalShaders.cpp b/Engine/Source/Runtime/IOS/MetalRHI/Private/MetalShaders.cpp
--- a/Engine/Source/Runtime/IOS/MetalRHI/Private/MetalShaders.cpp
+++ b/Engine/Source/Runtime/IOS/MetalRHI/Private/MetalShaders.cpp
@@ -77,7 +77,7 @@ TMetalBaseShader<BaseResourceType>::TMetalBaseShader(const TArray<uint8>& InCode
 	}
 	else
 	{
-		NSLog(@"Loaded an offline shader");
+		UE_LOG(LogMetal, Display, TEXT("Loaded a non-offline compiled shader (will be slower to load)"));
 
 		CodeLength = InCode.Num() - CodeOffset - 1;
 
diff --git a/Engine/Source/Runtime/IOS/MetalRHI/Private/MetalState.cpp b/Engine/Source/Runtime/IOS/MetalRHI/Private/MetalState.cpp
--- a/Engine/Source/Runtime/IOS/MetalRHI/Private/MetalState.cpp
+++ b/Engine/Source/Runtime/IOS/MetalRHI/Private/MetalState.cpp
@@ -100,6 +100,45 @@ static MTLStencilOperation TranslateStencilOp(EStencilOp StencilOp)
 	};
 }
 
+static MTLBlendOperation TranslateBlendOp(EBlendOperation BlendOp)
+{
+	switch(BlendOp)
+	{
+		case BO_Subtract:	return MTLBlendOperationSubtract;
+		case BO_Min:		return MTLBlendOperationMin;
+		case BO_Max:		return MTLBlendOperationMax;
+		default:			return MTLBlendOperationAdd;
+	};
+}
+
+
+static MTLBlendFactor TranslateBlendFactor(EBlendFactor BlendFactor)
+{
+	switch(BlendFactor)
+	{
+		case BF_One:					return MTLBlendFactorOne;
+		case BF_SourceColor:			return MTLBlendFactorSourceColor;
+		case BF_InverseSourceColor:		return MTLBlendFactorOneMinusSourceColor;
+		case BF_SourceAlpha:			return MTLBlendFactorSourceAlpha;
+		case BF_InverseSourceAlpha:		return MTLBlendFactorOneMinusSourceAlpha;
+		case BF_DestAlpha:				return MTLBlendFactorDestinationAlpha;
+		case BF_InverseDestAlpha:		return MTLBlendFactorOneMinusDestinationAlpha;
+		case BF_DestColor:				return MTLBlendFactorDestinationColor;
+		case BF_InverseDestColor:		return MTLBlendFactorOneMinusDestinationColor;
+		default:						return MTLBlendFactorZero;
+	};
+}
+
+static MTLColorWriteMask TranslateWriteMask(EColorWriteMask WriteMask)
+{
+	uint32 Result = 0;
+	Result |= (WriteMask & CW_RED) ? (MTLColorWriteMaskRed) : 0;
+	Result |= (WriteMask & CW_GREEN) ? (MTLColorWriteMaskGreen) : 0;
+	Result |= (WriteMask & CW_BLUE) ? (MTLColorWriteMaskBlue) : 0;
+	Result |= (WriteMask & CW_ALPHA) ? (MTLColorWriteMaskAlpha) : 0;
+	
+	return (MTLColorWriteMask)Result;
+}
 
 
 
@@ -208,17 +247,42 @@ void FMetalDepthStencilState::Set()
 }
 
 FMetalBlendState::FMetalBlendState(const FBlendStateInitializerRHI& Initializer)
-	: RenderTargetStates(Initializer)
 {
+	for(uint32 RenderTargetIndex = 0;RenderTargetIndex < MaxMetalRenderTargets; ++RenderTargetIndex)
+	{
+		// which initializer to use
+		const FBlendStateInitializerRHI::FRenderTarget& Init =
+			Initializer.bUseIndependentRenderTargetBlendStates
+				? Initializer.RenderTargets[RenderTargetIndex]
+				: Initializer.RenderTargets[0];
+
+		// make a new blend state
+		MTLRenderPipelineColorAttachmentDescriptor* BlendState = [[MTLRenderPipelineColorAttachmentDescriptor alloc] init];
+		TRACK_OBJECT(BlendState);
+
+		// set values
+		BlendState.blendingEnabled =
+			Init.ColorBlendOp != BO_Add || Init.ColorDestBlend != BF_Zero || Init.ColorSrcBlend != BF_One ||
+			Init.AlphaBlendOp != BO_Add || Init.AlphaDestBlend != BF_Zero || Init.AlphaSrcBlend != BF_One;
+		BlendState.sourceRGBBlendFactor = TranslateBlendFactor(Init.ColorSrcBlend);
+		BlendState.destinationRGBBlendFactor = TranslateBlendFactor(Init.ColorDestBlend);
+		BlendState.rgbBlendOperation = TranslateBlendOp(Init.ColorBlendOp);
+		BlendState.sourceAlphaBlendFactor = TranslateBlendFactor(Init.AlphaSrcBlend);
+		BlendState.destinationAlphaBlendFactor = TranslateBlendFactor(Init.AlphaDestBlend);
+		BlendState.alphaBlendOperation = TranslateBlendOp(Init.AlphaBlendOp);
+		BlendState.writeMask = TranslateWriteMask(Init.ColorWriteMask);
+		
+		RenderTargetStates[RenderTargetIndex] = BlendState;
+	}
 }
 
 FMetalBlendState::~FMetalBlendState()
 {
-// 	for(uint32 RenderTargetIndex = 0;RenderTargetIndex < MaxMetalRenderTargets; ++RenderTargetIndex)
-// 	{
-// 		UNTRACK_OBJECT(RenderTargetStates[RenderTargetIndex]);
-// 		[RenderTargetStates[RenderTargetIndex] release];
-// 	}
+	for(uint32 RenderTargetIndex = 0;RenderTargetIndex < MaxMetalRenderTargets; ++RenderTargetIndex)
+	{
+		UNTRACK_OBJECT(RenderTargetStates[RenderTargetIndex]);
+		[RenderTargetStates[RenderTargetIndex] release];
+	}
 }
 
 void FMetalBlendState::Set()
diff --git a/Engine/Source/Runtime/IOS/MetalRHI/Private/MetalVertexDeclaration.cpp b/Engine/Source/Runtime/IOS/MetalRHI/Private/MetalVertexDeclaration.cpp
--- a/Engine/Source/Runtime/IOS/MetalRHI/Private/MetalVertexDeclaration.cpp
+++ b/Engine/Source/Runtime/IOS/MetalRHI/Private/MetalVertexDeclaration.cpp
@@ -114,7 +114,9 @@ void FMetalVertexDeclaration::GenerateLayout(const FVertexDeclarationElementList
 			}
 
 			// set the stride once per buffer
-			[Layout setStride:Stride stepFunction:Function stepRate:StepRate atVertexBufferIndex:ShaderBufferIndex];
+			Layout.layouts[ShaderBufferIndex].stride = Stride;
+			Layout.layouts[ShaderBufferIndex].stepFunction = Function;
+			Layout.layouts[ShaderBufferIndex].stepRate = StepRate;
 
 			// track this buffer and stride
 			BufferStrides.Add(ShaderBufferIndex, Element.Stride);
@@ -126,6 +128,8 @@ void FMetalVertexDeclaration::GenerateLayout(const FVertexDeclarationElementList
 		}
 
 		// set the format for each element
-		[Layout setVertexFormat:TranslateElementTypeToMTLType(Element.Type) offset:Element.Offset vertexBufferIndex:ShaderBufferIndex atAttributeIndex:Element.AttributeIndex];
+		Layout.attributes[Element.AttributeIndex].format = TranslateElementTypeToMTLType(Element.Type);
+		Layout.attributes[Element.AttributeIndex].offset = Element.Offset;
+		Layout.attributes[Element.AttributeIndex].bufferIndex = ShaderBufferIndex;
 	}
 }
diff --git a/Engine/Source/Runtime/IOS/MetalRHI/Public/MetalState.h b/Engine/Source/Runtime/IOS/MetalRHI/Public/MetalState.h
--- a/Engine/Source/Runtime/IOS/MetalRHI/Public/MetalState.h
+++ b/Engine/Source/Runtime/IOS/MetalRHI/Public/MetalState.h
@@ -75,5 +75,5 @@ public:
 	 */
 	void Set();
 	
-	FBlendStateInitializerRHI RenderTargetStates;
+	TStaticArray<MTLRenderPipelineColorAttachmentDescriptor*, MaxSimultaneousRenderTargets> RenderTargetStates;
 };
