UE4 - Fix RHI featurelevel for Metal

diff --git a/Engine/Source/Runtime/IOS/MetalRHI/Private/MetalRHI.cpp b/Engine/Source/Runtime/IOS/MetalRHI/Private/MetalRHI.cpp
--- a/Engine/Source/Runtime/IOS/MetalRHI/Private/MetalRHI.cpp
+++ b/Engine/Source/Runtime/IOS/MetalRHI/Private/MetalRHI.cpp
@@ -33,6 +33,7 @@ FMetalDynamicRHI::FMetalDynamicRHI()
 	// Initialize the RHI capabilities.
 	GRHIFeatureLevel = ERHIFeatureLevel::ES2;
 	GRHIShaderPlatform = SP_METAL;
+	GMaxRHIFeatureLevel = GCurrentRHIFeatureLevel = ERHIFeatureLevel::ES2;
 
 	GPixelCenterOffset = 0.0f;
 	GSupportsVertexInstancing = false;
