UE4 - Metal format fix

diff --git a/Engine/Source/Runtime/IOS/MetalRHI/Private/MetalRHI.cpp b/Engine/Source/Runtime/IOS/MetalRHI/Private/MetalRHI.cpp
--- a/Engine/Source/Runtime/IOS/MetalRHI/Private/MetalRHI.cpp
+++ b/Engine/Source/Runtime/IOS/MetalRHI/Private/MetalRHI.cpp
@@ -94,7 +94,7 @@ FMetalDynamicRHI::FMetalDynamicRHI()
 	GPixelFormats[PF_G16R16F_FILTER		].PlatformFormat	= MTLPixelFormatRG16Float;
 	GPixelFormats[PF_G32R32F			].PlatformFormat	= MTLPixelFormatRG32Float;
 	GPixelFormats[PF_A2B10G10R10		].PlatformFormat    = MTLPixelFormatRGB10A2Unorm;
-	GPixelFormats[PF_A16B16G16R16		].PlatformFormat    = MTLPixelFormatRGBA16Float;
+	GPixelFormats[PF_A16B16G16R16		].PlatformFormat    = MTLPixelFormatRGBA16Unorm;
 	GPixelFormats[PF_D24				].PlatformFormat	= MTLPixelFormatInvalid;
 	GPixelFormats[PF_R16F				].PlatformFormat	= MTLPixelFormatR16Float;
 	GPixelFormats[PF_R16F_FILTER		].PlatformFormat	= MTLPixelFormatR16Float;
