- Early code for MRT in Metal (supporst more than one, but untested, no use case yet)

diff --git a/Engine/Source/Runtime/IOS/MetalRHI/Private/MetalCommands.cpp b/Engine/Source/Runtime/IOS/MetalRHI/Private/MetalCommands.cpp
--- a/Engine/Source/Runtime/IOS/MetalRHI/Private/MetalCommands.cpp
+++ b/Engine/Source/Runtime/IOS/MetalRHI/Private/MetalCommands.cpp
@@ -302,18 +302,19 @@ void FMetalDynamicRHI::RHISetRenderTargets(uint32 NumSimultaneousRenderTargets,
 {
 	FMetalManager* Manager = FMetalManager::Get();
 	
-	// only support 1 for now!!!
-	check(NumSimultaneousRenderTargets <= 1);
-	
-	// update the current RTs
-	if (NumSimultaneousRenderTargets && NewRenderTargets[0].Texture != NULL)
-	{
-		FMetalSurface& Surface = GetMetalSurfaceFromRHITexture(NewRenderTargets[0].Texture);
-		Manager->SetCurrentRenderTarget(&Surface);
-	}
-	else
+
+	for (int32 RenderTargetIndex = 0; RenderTargetIndex < MaxMetalRenderTargets; RenderTargetIndex++)
 	{
-		Manager->SetCurrentRenderTarget(NULL);
+		// update the current RTs
+		if (RenderTargetIndex < NumSimultaneousRenderTargets && NewRenderTargets[RenderTargetIndex].Texture != NULL)
+		{
+			FMetalSurface& Surface = GetMetalSurfaceFromRHITexture(NewRenderTargets[0].Texture);
+			Manager->SetCurrentRenderTarget(&Surface, RenderTargetIndex, NumSimultaneousRenderTargets);
+		}
+		else
+		{
+			Manager->SetCurrentRenderTarget(NULL, RenderTargetIndex, NumSimultaneousRenderTargets);
+		}
 	}
 	
 	if (NewDepthStencilTargetRHI)
diff --git a/Engine/Source/Runtime/IOS/MetalRHI/Private/MetalManager.cpp b/Engine/Source/Runtime/IOS/MetalRHI/Private/MetalManager.cpp
--- a/Engine/Source/Runtime/IOS/MetalRHI/Private/MetalManager.cpp
+++ b/Engine/Source/Runtime/IOS/MetalRHI/Private/MetalManager.cpp
@@ -9,28 +9,25 @@ const uint32 RingBufferSize = 8 * 1024 * 1024;
 TMap<id, int32> ClassCounts;
 #endif
 
-#define NUMBITS_SOURCE_RGB_BLEND_FACTOR			4
-#define NUMBITS_DEST_RGB_BLEND_FACTOR			4
-#define NUMBITS_RGB_BLEND_OPERATION				3
-#define NUMBITS_SOURCE_A_BLEND_FACTOR			4
-#define NUMBITS_DEST_A_BLEND_FACTOR				4
-#define NUMBITS_A_BLEND_OPERATION				3
-#define NUMBITS_WRITE_MASK						4
-#define NUMBITS_RENDER_TARGET_FORMAT			9
-#define NUMBITS_DEPTH_TARGET_FORMAT				9
-#define NUMBITS_SAMPLE_COUNT					3
-
-
-#define OFFSET_SOURCE_RGB_BLEND_FACTOR			(0)
-#define OFFSET_DEST_RGB_BLEND_FACTOR			(OFFSET_SOURCE_RGB_BLEND_FACTOR	+ NUMBITS_SOURCE_RGB_BLEND_FACTOR)
-#define OFFSET_RGB_BLEND_OPERATION				(OFFSET_DEST_RGB_BLEND_FACTOR	+ NUMBITS_DEST_RGB_BLEND_FACTOR)
-#define OFFSET_SOURCE_A_BLEND_FACTOR			(OFFSET_RGB_BLEND_OPERATION		+ NUMBITS_RGB_BLEND_OPERATION)
-#define OFFSET_DEST_A_BLEND_FACTOR				(OFFSET_SOURCE_A_BLEND_FACTOR	+ NUMBITS_SOURCE_A_BLEND_FACTOR)
-#define OFFSET_A_BLEND_OPERATION				(OFFSET_DEST_A_BLEND_FACTOR		+ NUMBITS_DEST_A_BLEND_FACTOR)
-#define OFFSET_WRITE_MASK						(OFFSET_A_BLEND_OPERATION		+ NUMBITS_A_BLEND_OPERATION)
-#define OFFSET_RENDER_TARGET_FORMAT				(OFFSET_WRITE_MASK				+ NUMBITS_WRITE_MASK)
-#define OFFSET_DEPTH_TARGET_FORMAT				(OFFSET_RENDER_TARGET_FORMAT	+ NUMBITS_RENDER_TARGET_FORMAT)
-#define OFFSET_SAMPLE_COUNT						(OFFSET_DEPTH_TARGET_FORMAT		+ NUMBITS_DEPTH_TARGET_FORMAT)
+#define NUMBITS_BLEND_STATE				5
+#define NUMBITS_RENDER_TARGET_FORMAT	8
+#define NUMBITS_DEPTH_TARGET_FORMAT		8
+#define NUMBITS_SAMPLE_COUNT			3
+
+
+#define OFFSET_BLEND_STATE0				(0)
+#define OFFSET_BLEND_STATE1				(OFFSET_BLEND_STATE0	+ NUMBITS_BLEND_STATE)
+#define OFFSET_BLEND_STATE2				(OFFSET_BLEND_STATE1	+ NUMBITS_BLEND_STATE)
+#define OFFSET_BLEND_STATE3				(OFFSET_BLEND_STATE2	+ NUMBITS_BLEND_STATE)
+#define OFFSET_RENDER_TARGET_FORMAT0	(OFFSET_BLEND_STATE3	+ NUMBITS_BLEND_STATE)
+#define OFFSET_RENDER_TARGET_FORMAT1	(OFFSET_RENDER_TARGET_FORMAT0	+ NUMBITS_RENDER_TARGET_FORMAT)
+#define OFFSET_RENDER_TARGET_FORMAT2	(OFFSET_RENDER_TARGET_FORMAT1	+ NUMBITS_RENDER_TARGET_FORMAT)
+#define OFFSET_RENDER_TARGET_FORMAT3	(OFFSET_RENDER_TARGET_FORMAT2	+ NUMBITS_RENDER_TARGET_FORMAT)
+#define OFFSET_DEPTH_TARGET_FORMAT		(OFFSET_RENDER_TARGET_FORMAT3	+ NUMBITS_RENDER_TARGET_FORMAT)
+#define OFFSET_SAMPLE_COUNT				(OFFSET_DEPTH_TARGET_FORMAT		+ NUMBITS_DEPTH_TARGET_FORMAT)
+
+static uint32 BlendBitOffsets[] = { OFFSET_BLEND_STATE0, OFFSET_BLEND_STATE1, OFFSET_BLEND_STATE2, OFFSET_BLEND_STATE3, };
+static uint32 RTBitOffsets[] = { OFFSET_RENDER_TARGET_FORMAT0, OFFSET_RENDER_TARGET_FORMAT1, OFFSET_RENDER_TARGET_FORMAT2, OFFSET_RENDER_TARGET_FORMAT3, };
 
 #define SET_HASH(Offset, NumBits, Value) \
 	{ \
@@ -39,17 +36,10 @@ TMap<id, int32> ClassCounts;
 	}
 #define GET_HASH(Offset, NumBits) ((Hash >> Offset) & ((1ULL << NumBits) - 1))
 
-DEFINE_STAT(STAT_MetalMakeDrawableTime);
-DEFINE_STAT(STAT_MetalDrawCallTime);
-DEFINE_STAT(STAT_MetalPrepareDrawTime);
-DEFINE_STAT(STAT_MetalUniformBufferCleanupTime);
-DEFINE_STAT(STAT_MetalFreeUniformBufferMemory);
-DEFINE_STAT(STAT_MetalNumFreeUniformBuffers);
-DEFINE_STAT(STAT_MetalPipelineStateTime);
-DEFINE_STAT(STAT_MetalBoundShaderStateTime);
-DEFINE_STAT(STAT_MetalVertexDeclarationTime);
+
 
 	
+/*
 void FPipelineShadow::SetHash(uint64 InHash)
 {
 	Hash = InHash;
@@ -71,6 +61,7 @@ void FPipelineShadow::SetHash(uint64 InHash)
 	DepthTargetFormat = (MTLPixelFormat)GET_HASH(OFFSET_DEPTH_TARGET_FORMAT, NUMBITS_DEPTH_TARGET_FORMAT);
 	SampleCount = GET_HASH(OFFSET_SAMPLE_COUNT, NUMBITS_SAMPLE_COUNT);
 }
+*/
 
 
 
@@ -199,8 +190,8 @@ FMetalManager::FMetalManager()
 	, CurrentCommandBuffer(nil)
 	, CurrentDrawable(nil)
 	, CurrentContext(nil)
-	, CurrentColorRenderTexture(nil)
-	, PreviousColorRenderTexture(nil)
+	, CurrentNumRenderTargets(0)
+	, PreviousNumRenderTargets(0)
 	, CurrentDepthRenderTexture(nil)
 	, PreviousDepthRenderTexture(nil)
 	, RingBuffer(Device, RingBufferSize, BufferOffsetAlignment)
@@ -211,6 +202,12 @@ FMetalManager::FMetalManager()
 	, SceneFrameCounter(0)
 	, ResourceTableFrameCounter(INDEX_NONE)
 {
+	for (int32 Index = 0; Index < ARRAY_COUNT(CurrentColorRenderTextures); Index++)
+	{
+		CurrentColorRenderTextures[Index] = nil;
+		PreviousColorRenderTextures[Index] = nil;
+	}
+
 	CommandQueue = [Device newCommandQueue];
 
 	// get the size of the window
@@ -404,30 +401,21 @@ void FMetalManager::SetBlendState(FMetalBlendState* BlendState)
 {
 	for(uint32 RenderTargetIndex = 0;RenderTargetIndex < MaxMetalRenderTargets; ++RenderTargetIndex)
 	{
-		// @todo metal mrt
-		MTLRenderPipelineColorAttachmentDescriptor* Blend0 = BlendState->RenderTargetStates[RenderTargetIndex];
+		MTLRenderPipelineColorAttachmentDescriptor* Blend = BlendState->RenderTargetStates[RenderTargetIndex].BlendState;
 		MTLRenderPipelineColorAttachmentDescriptor* Dest = Pipeline.RenderTargets[RenderTargetIndex];
-		check(RenderTargetIndex == 0);
 
-		// assign the struct by value
-//		*Pipeline.RenderTargets[RenderTargetIndex] = *Blend0;
-		Dest.blendingEnabled = Blend0.blendingEnabled;
-		Dest.sourceRGBBlendFactor = Blend0.sourceRGBBlendFactor;
-		Dest.destinationRGBBlendFactor = Blend0.destinationRGBBlendFactor;
-		Dest.rgbBlendOperation = Blend0.rgbBlendOperation;
-		Dest.sourceAlphaBlendFactor = Blend0.sourceAlphaBlendFactor;
-		Dest.destinationAlphaBlendFactor = Blend0.destinationAlphaBlendFactor;
-		Dest.alphaBlendOperation = Blend0.alphaBlendOperation;
-		Dest.writeMask = Blend0.writeMask;
+		// assign each property manually, would be nice if this was faster
+		Dest.blendingEnabled = Blend.blendingEnabled;
+		Dest.sourceRGBBlendFactor = Blend.sourceRGBBlendFactor;
+		Dest.destinationRGBBlendFactor = Blend.destinationRGBBlendFactor;
+		Dest.rgbBlendOperation = Blend.rgbBlendOperation;
+		Dest.sourceAlphaBlendFactor = Blend.sourceAlphaBlendFactor;
+		Dest.destinationAlphaBlendFactor = Blend.destinationAlphaBlendFactor;
+		Dest.alphaBlendOperation = Blend.alphaBlendOperation;
+		Dest.writeMask = Blend.writeMask;
 
-
-		SET_HASH(OFFSET_SOURCE_RGB_BLEND_FACTOR, NUMBITS_SOURCE_RGB_BLEND_FACTOR, Blend0.sourceRGBBlendFactor);
-		SET_HASH(OFFSET_DEST_RGB_BLEND_FACTOR, NUMBITS_DEST_RGB_BLEND_FACTOR, Blend0.destinationRGBBlendFactor);
-		SET_HASH(OFFSET_RGB_BLEND_OPERATION, NUMBITS_RGB_BLEND_OPERATION, Blend0.rgbBlendOperation);
-		SET_HASH(OFFSET_SOURCE_A_BLEND_FACTOR, NUMBITS_SOURCE_A_BLEND_FACTOR, Blend0.sourceAlphaBlendFactor);
-		SET_HASH(OFFSET_DEST_A_BLEND_FACTOR, NUMBITS_DEST_A_BLEND_FACTOR, Blend0.destinationAlphaBlendFactor);
-		SET_HASH(OFFSET_A_BLEND_OPERATION, NUMBITS_A_BLEND_OPERATION, Blend0.alphaBlendOperation);
-		SET_HASH(OFFSET_WRITE_MASK, NUMBITS_WRITE_MASK, Blend0.writeMask);
+		// set the hash bits for this RT
+		SET_HASH(BlendBitOffsets[RenderTargetIndex], NUMBITS_BLEND_STATE, BlendState->RenderTargetStates[RenderTargetIndex].BlendStateKey);
 	}
 }
 
@@ -439,8 +427,11 @@ void FMetalManager::SetBoundShaderState(FMetalBoundShaderState* BoundShaderState
 	CurrentBoundShaderState = BoundShaderState;
 }
 
-void FMetalManager::SetCurrentRenderTarget(FMetalSurface* RenderSurface)
+void FMetalManager::SetCurrentRenderTarget(FMetalSurface* RenderSurface, int32 RenderTargetIndex, int32 TotalNumRenderTargets)
 {
+	// rememeber our new max
+	CurrentNumRenderTargets = TotalNumRenderTargets;
+
 	// update the current rendered-to pixel format
 	if (RenderSurface)
 	{
@@ -451,8 +442,8 @@ void FMetalManager::SetCurrentRenderTarget(FMetalSurface* RenderSurface)
 
 			uint32 IdleStart = FPlatformTime::Cycles();
 
-				// make a drawable object for this frame
-				CurrentDrawable = [[IOSAppDelegate GetDelegate].IOSView MakeDrawable];
+			// make a drawable object for this frame
+			CurrentDrawable = [[IOSAppDelegate GetDelegate].IOSView MakeDrawable];
 
 			GRenderThreadIdle[ERenderThreadIdleTypes::WaitingForGPUPresent] += FPlatformTime::Cycles() - IdleStart;
 			GRenderThreadNumIdle[ERenderThreadIdleTypes::WaitingForGPUPresent]++;
@@ -461,12 +452,15 @@ void FMetalManager::SetCurrentRenderTarget(FMetalSurface* RenderSurface)
 			RenderSurface->Texture = CurrentDrawable.texture;
 		}
 
-		CurrentColorRenderTexture = RenderSurface->Texture;
+		CurrentColorRenderTextures[RenderTargetIndex] = RenderSurface->Texture;
+
+		// only allow one MRT when using MSAA
+		checkf(RenderSurface->MSAATexture == NULL || TotalNumRenderTargets == 1);
 		CurrentMSAARenderTexture = RenderSurface->MSAATexture;
 	}
 	else
 	{
-		CurrentColorRenderTexture = nil;
+		CurrentColorRenderTextures[RenderTargetIndex] = nil;
 	}
 }
 
@@ -485,14 +479,30 @@ void FMetalManager::SetCurrentDepthStencilTarget(FMetalSurface* RenderSurface)
 
 void FMetalManager::UpdateContext()
 {
-	if (CurrentColorRenderTexture == PreviousColorRenderTexture && CurrentDepthRenderTexture == PreviousDepthRenderTexture)
+	// if all render targets match, we can early out
+	if (CurrentNumRenderTargets == PreviousNumRenderTargets && CurrentDepthRenderTexture == PreviousDepthRenderTexture)
 	{
-		return;
+		// make sure all match
+		bool bAllMatch = true;
+		for (uint32 AttachmentIndex = 0; AttachmentIndex < CurrentNumRenderTargets; AttachmentIndex++)
+		{
+			if (CurrentColorRenderTextures[AttachmentIndex] != PreviousColorRenderTextures[AttachmentIndex])
+			{
+				bAllMatch = false;
+			}
+		}
+
+		if (bAllMatch)
+		{
+			return;
+		}
 	}
 
+	PreviousNumRenderTargets = CurrentNumRenderTargets;
+
 	// if we are setting them to nothing, then this is probably end of frame, and we can't make a framebuffer
-	// with nothng, so just abort this
-	if (CurrentColorRenderTexture == nil && CurrentDepthRenderTexture == nil)
+	// with nothng, so just abort this (only need to check on single MRT case)
+	if (CurrentNumRenderTargets == 1 && CurrentColorRenderTextures[0] == nil && CurrentDepthRenderTexture == nil)
 	{
 		return;
 	}
@@ -503,50 +513,60 @@ void FMetalManager::UpdateContext()
 	// if we need to do queries, write to the ring buffer (we set the offset into the ring buffer per query)
 	CurrentRenderPass.visibilityResultBuffer = QueryBuffer.Buffer;
 
-	// @todo metal mrt: Set an attachment index properly
-	uint32 AttachmentIndex = 0;
-
-	MTLRenderPassColorAttachmentDescriptor* ColorAttachment = nil;
-	MTLRenderPassDepthAttachmentDescriptor* DepthAttachment = nil;
-	// @todo metal stencil: handle stencil
-	MTLRenderPassAttachmentDescriptor* StencilAttachment = nil;
-
+	// default to non-msaa
 	Pipeline.SampleCount = 0;
-	if (CurrentColorRenderTexture)
-	{
-		ColorAttachment = [[MTLRenderPassColorAttachmentDescriptor alloc] init];
 
-		if (CurrentMSAARenderTexture)
+	for (uint32 AttachmentIndex = 0; AttachmentIndex < MaxMetalRenderTargets; AttachmentIndex++)
+	{
+		// only try to set it if it was one that was set (ie less than CurrentNumRenderTargets)
+		if (AttachmentIndex < CurrentNumRenderTargets && CurrentColorRenderTextures[AttachmentIndex])
 		{
-			ColorAttachment.texture = CurrentMSAARenderTexture;
-			[ColorAttachment setLoadAction:MTLLoadActionDontCare];
-			[ColorAttachment setStoreAction:MTLStoreActionMultisampleResolve];
-			[ColorAttachment setResolveTexture:CurrentColorRenderTexture];
-			Pipeline.SampleCount = CurrentMSAARenderTexture.sampleCount;
+			MTLRenderPassColorAttachmentDescriptor* ColorAttachment = [[MTLRenderPassColorAttachmentDescriptor alloc] init];
+
+			if (CurrentMSAARenderTexture)
+			{
+				// set up an MSAA attachment
+				ColorAttachment.texture = CurrentMSAARenderTexture;
+				[ColorAttachment setLoadAction:MTLLoadActionDontCare];
+				[ColorAttachment setStoreAction:MTLStoreActionMultisampleResolve];
+				[ColorAttachment setResolveTexture:CurrentColorRenderTextures[AttachmentIndex]];
+				Pipeline.SampleCount = CurrentMSAARenderTexture.sampleCount;
+
+				// only allow one MRT with msaa
+				checkf(CurrentNumRenderTargets == 1, TEXT("Only expected one MRT when using MSAA"));
+			}
+			else
+			{
+				// set up non-MSAA attachment
+				ColorAttachment.texture = CurrentColorRenderTextures[AttachmentIndex];
+				[ColorAttachment setLoadAction:MTLLoadActionDontCare];
+				[ColorAttachment setStoreAction:MTLStoreActionStore];
+				Pipeline.SampleCount = 1;
+			}
+
+			// assign the attachment to the slot
+			[CurrentRenderPass.colorAttachments setObject:ColorAttachment atIndexedSubscript:AttachmentIndex];
+			[ColorAttachment release];
+
+			Pipeline.RenderTargets[AttachmentIndex].pixelFormat = CurrentColorRenderTextures[AttachmentIndex].pixelFormat;
 		}
 		else
 		{
-			ColorAttachment.texture = CurrentColorRenderTexture;
-			[ColorAttachment setLoadAction:MTLLoadActionDontCare];
-			[ColorAttachment setStoreAction:MTLStoreActionStore];
-			Pipeline.SampleCount = 1;
+			Pipeline.RenderTargets[AttachmentIndex].pixelFormat = MTLPixelFormatInvalid;
 		}
+		
+		// update the hash no matter what case (null, unused, used)
+		SET_HASH(RTBitOffsets[AttachmentIndex], NUMBITS_RENDER_TARGET_FORMAT, Pipeline.RenderTargets[AttachmentIndex].pixelFormat);
 
-		[CurrentRenderPass.colorAttachments setObject:ColorAttachment atIndexedSubscript:AttachmentIndex];
-		[ColorAttachment release];
-
-		Pipeline.RenderTargets[AttachmentIndex].pixelFormat = CurrentColorRenderTexture.pixelFormat;
+		// remember this for next time
+		PreviousColorRenderTextures[AttachmentIndex] = CurrentColorRenderTextures[AttachmentIndex];
 	}
-	else
-	{
-		Pipeline.RenderTargets[AttachmentIndex].pixelFormat = MTLPixelFormatInvalid;
-	}
-	SET_HASH(OFFSET_RENDER_TARGET_FORMAT, NUMBITS_RENDER_TARGET_FORMAT, Pipeline.RenderTargets[AttachmentIndex].pixelFormat);
 	
 	if (CurrentDepthRenderTexture)
 	{
-		DepthAttachment = [[MTLRenderPassDepthAttachmentDescriptor alloc] init];
+		MTLRenderPassDepthAttachmentDescriptor* DepthAttachment = [[MTLRenderPassDepthAttachmentDescriptor alloc] init];
 
+		// set up the depth attachment
 		DepthAttachment.texture = CurrentDepthRenderTexture;
 		[DepthAttachment setLoadAction:MTLLoadActionClear];
 		[DepthAttachment setStoreAction:MTLStoreActionDontCare];
@@ -558,6 +578,7 @@ void FMetalManager::UpdateContext()
 			Pipeline.SampleCount = CurrentDepthRenderTexture.sampleCount;
 		}
 
+		// and assign it
 		CurrentRenderPass.depthAttachment = DepthAttachment;
 		[DepthAttachment release];
 	}
@@ -565,12 +586,14 @@ void FMetalManager::UpdateContext()
 	{
 		Pipeline.DepthTargetFormat = MTLPixelFormatInvalid;
 	}
+	// update hash for the depth buffer
 	SET_HASH(OFFSET_DEPTH_TARGET_FORMAT, NUMBITS_DEPTH_TARGET_FORMAT, Pipeline.DepthTargetFormat);
-	SET_HASH(OFFSET_SAMPLE_COUNT, NUMBITS_SAMPLE_COUNT, Pipeline.SampleCount);
 
-	PreviousColorRenderTexture = CurrentColorRenderTexture;
+	// remember this for next time
 	PreviousDepthRenderTexture = CurrentDepthRenderTexture;
 
+	SET_HASH(OFFSET_SAMPLE_COUNT, NUMBITS_SAMPLE_COUNT, Pipeline.SampleCount);
+
 	// commit pending commands on the old render target
 	if (CurrentContext)
 	{
diff --git a/Engine/Source/Runtime/IOS/MetalRHI/Private/MetalManager.h b/Engine/Source/Runtime/IOS/MetalRHI/Private/MetalManager.h
--- a/Engine/Source/Runtime/IOS/MetalRHI/Private/MetalManager.h
+++ b/Engine/Source/Runtime/IOS/MetalRHI/Private/MetalManager.h
@@ -29,10 +29,10 @@ struct FPipelineShadow
 		return Hash;
 	}
 
-	/**
-	 * FOrce a pipeline state for a known hash
-	 */
-	void SetHash(uint64 InHash);
+// 	/**
+// 	 * FOrce a pipeline state for a known hash
+// 	 */
+// 	void SetHash(uint64 InHash);
 	
 	/**
 	 * @return an pipeline state object that matches the current state and the BSS
@@ -99,7 +99,7 @@ public:
 	 */
 	void SetBlendState(FMetalBlendState* BlendState);
 	void SetBoundShaderState(FMetalBoundShaderState* BoundShaderState);
-	void SetCurrentRenderTarget(FMetalSurface* RenderSurface);
+	void SetCurrentRenderTarget(FMetalSurface* RenderSurface, int32 RenderTargetIndex, int32 TotalNumRenderTargets);
 	void SetCurrentDepthStencilTarget(FMetalSurface* RenderSurface);
 	
 	/**
@@ -164,7 +164,8 @@ protected:
 
 	id<MTLRenderCommandEncoder> CurrentContext;
 	
-	id<MTLTexture> CurrentColorRenderTexture, PreviousColorRenderTexture;
+	uint32 CurrentNumRenderTargets, PreviousNumRenderTargets;
+	id<MTLTexture> CurrentColorRenderTextures[MaxMetalRenderTargets], PreviousColorRenderTextures[MaxMetalRenderTargets];
 	id<MTLTexture> CurrentDepthRenderTexture, PreviousDepthRenderTexture;
 	id<MTLTexture> CurrentMSAARenderTexture;
 	
diff --git a/Engine/Source/Runtime/IOS/MetalRHI/Private/MetalRHI.cpp b/Engine/Source/Runtime/IOS/MetalRHI/Private/MetalRHI.cpp
--- a/Engine/Source/Runtime/IOS/MetalRHI/Private/MetalRHI.cpp
+++ b/Engine/Source/Runtime/IOS/MetalRHI/Private/MetalRHI.cpp
@@ -12,6 +12,17 @@
 
 DEFINE_LOG_CATEGORY(LogMetal)
 
+DEFINE_STAT(STAT_MetalMakeDrawableTime);
+DEFINE_STAT(STAT_MetalDrawCallTime);
+DEFINE_STAT(STAT_MetalPrepareDrawTime);
+DEFINE_STAT(STAT_MetalUniformBufferCleanupTime);
+DEFINE_STAT(STAT_MetalFreeUniformBufferMemory);
+DEFINE_STAT(STAT_MetalNumFreeUniformBuffers);
+DEFINE_STAT(STAT_MetalPipelineStateTime);
+DEFINE_STAT(STAT_MetalBoundShaderStateTime);
+DEFINE_STAT(STAT_MetalVertexDeclarationTime);
+
+
 bool FMetalDynamicRHIModule::IsSupported()
 {
 	return true;
diff --git a/Engine/Source/Runtime/IOS/MetalRHI/Private/MetalRHIPrivate.h b/Engine/Source/Runtime/IOS/MetalRHI/Private/MetalRHIPrivate.h
--- a/Engine/Source/Runtime/IOS/MetalRHI/Private/MetalRHIPrivate.h
+++ b/Engine/Source/Runtime/IOS/MetalRHI/Private/MetalRHIPrivate.h
@@ -9,8 +9,7 @@
 #include "Engine.h"
 
 // UE4 has a Max of 8 RTs, but we max out at 4, no need to waste time processing 8
-// @todo metal mrt: Currently only using 1, so just do 1!
-const uint32 MaxMetalRenderTargets = 1;
+const uint32 MaxMetalRenderTargets = 4;
 
 // How many possible vertex streams are allowed
 const uint32 MaxMetalStreams = 16;
diff --git a/Engine/Source/Runtime/IOS/MetalRHI/Private/MetalState.cpp b/Engine/Source/Runtime/IOS/MetalRHI/Private/MetalState.cpp
--- a/Engine/Source/Runtime/IOS/MetalRHI/Private/MetalState.cpp
+++ b/Engine/Source/Runtime/IOS/MetalRHI/Private/MetalState.cpp
@@ -243,6 +243,13 @@ void FMetalDepthStencilState::Set()
 	[FMetalManager::GetContext() setDepthStencilState:State];
 }
 
+
+
+// statics
+TMap<uint32, uint8> FMetalBlendState::BlendSettingsToUniqueKeyMap;
+uint8 FMetalBlendState::NextKey = 0;
+
+
 FMetalBlendState::FMetalBlendState(const FBlendStateInitializerRHI& Initializer)
 {
 	for(uint32 RenderTargetIndex = 0;RenderTargetIndex < MaxMetalRenderTargets; ++RenderTargetIndex)
@@ -269,7 +276,24 @@ FMetalBlendState::FMetalBlendState(const FBlendStateInitializerRHI& Initializer)
 		BlendState.alphaBlendOperation = TranslateBlendOp(Init.AlphaBlendOp);
 		BlendState.writeMask = TranslateWriteMask(Init.ColorWriteMask);
 		
-		RenderTargetStates[RenderTargetIndex] = BlendState;
+		RenderTargetStates[RenderTargetIndex].BlendState = BlendState;
+
+		// get the unique key
+		uint32 BlendBitMask =
+			(BlendState.sourceRGBBlendFactor << 0) | (BlendState.destinationRGBBlendFactor << 4) | (BlendState.rgbBlendOperation << 8) |
+			(BlendState.sourceAlphaBlendFactor << 11) | (BlendState.destinationAlphaBlendFactor << 15) | (BlendState.alphaBlendOperation << 19) |
+			(BlendState.writeMask << 22);
+		uint8* Key = BlendSettingsToUniqueKeyMap.Find(BlendBitMask);
+		if (Key == NULL)
+		{
+			NSLog(@"Making a new key [BlendBitMask: %x, Key: %d]", BlendBitMask, NextKey);
+			Key = &BlendSettingsToUniqueKeyMap.Add(BlendBitMask, NextKey++);
+
+			// only giving 5 bits to the key, we need 4 of them
+			checkf(NextKey < 32, TEXT("Too many unique blend states to fit into the PipelineStateHash"));
+		}
+		// set the key
+		RenderTargetStates[RenderTargetIndex].BlendStateKey = *Key;
 	}
 }
 
@@ -278,7 +302,7 @@ FMetalBlendState::~FMetalBlendState()
 	for(uint32 RenderTargetIndex = 0;RenderTargetIndex < MaxMetalRenderTargets; ++RenderTargetIndex)
 	{
 		UNTRACK_OBJECT(RenderTargetStates[RenderTargetIndex]);
-		[RenderTargetStates[RenderTargetIndex] release];
+		[RenderTargetStates[RenderTargetIndex].BlendState release];
 	}
 }
 
diff --git a/Engine/Source/Runtime/IOS/MetalRHI/Public/MetalState.h b/Engine/Source/Runtime/IOS/MetalRHI/Public/MetalState.h
--- a/Engine/Source/Runtime/IOS/MetalRHI/Public/MetalState.h
+++ b/Engine/Source/Runtime/IOS/MetalRHI/Public/MetalState.h
@@ -75,5 +75,16 @@ public:
 	 */
 	void Set();
 	
-	TStaticArray<MTLRenderPipelineColorAttachmentDescriptor*, MaxSimultaneousRenderTargets> RenderTargetStates;
+
+	struct FBlendPerMRT
+	{
+		MTLRenderPipelineColorAttachmentDescriptor* BlendState;
+		uint8 BlendStateKey;
+	};
+	FBlendPerMRT RenderTargetStates[MaxMetalRenderTargets];
+
+private:
+	// this tracks blend settings (in a bit flag) into a unique key that uses few bits, for PipelineState MRT setup
+	static TMap<uint32, uint8> BlendSettingsToUniqueKeyMap;
+	static uint8 NextKey;
 };
