- Fixed Metal PCH compiling

diff --git a/Engine/Source/Runtime/IOS/MetalRHI/Private/MetalManager.cpp b/Engine/Source/Runtime/IOS/MetalRHI/Private/MetalManager.cpp
--- a/Engine/Source/Runtime/IOS/MetalRHI/Private/MetalManager.cpp
+++ b/Engine/Source/Runtime/IOS/MetalRHI/Private/MetalManager.cpp
@@ -1,6 +1,6 @@
 // Copyright 1998-2014 Epic Games, Inc. All Rights Reserved.
 
-#include "MetalManager.h"
+#include "MetalRHIPrivate.h"
 #include "IOSAppDelegate.h"
 
 const uint32 RingBufferSize = 8 * 1024 * 1024;
