- Texture streaming support in Metal

diff --git a/Engine/Source/Runtime/IOS/MetalRHI/Private/MetalManager.h b/Engine/Source/Runtime/IOS/MetalRHI/Private/MetalManager.h
--- a/Engine/Source/Runtime/IOS/MetalRHI/Private/MetalManager.h
+++ b/Engine/Source/Runtime/IOS/MetalRHI/Private/MetalManager.h
@@ -141,12 +141,22 @@ public:
 		return ShaderParameters[Stage];
 	}
 
+	/**
+	 * Return an auto-released command buffer, caller will need to retain it if it needs to live awhile
+	 */
+	id<MTLCommandBuffer> CreateTempCommandBuffer(bool bRetainReferences)
+	{
+		return bRetainReferences ? [CommandQueue commandBuffer] : [CommandQueue commandBufferWithUnretainedReferences];
+	}
+
 	/**
 	 * Handle rendering thread starting/stopping
 	 */
 	void CreateAutoreleasePool();
 	void DrainAutoreleasePool();
 
+
+
 protected:
 	FMetalManager();
 	void InitFrame();
diff --git a/Engine/Source/Runtime/IOS/MetalRHI/Private/MetalRHI.cpp b/Engine/Source/Runtime/IOS/MetalRHI/Private/MetalRHI.cpp
--- a/Engine/Source/Runtime/IOS/MetalRHI/Private/MetalRHI.cpp
+++ b/Engine/Source/Runtime/IOS/MetalRHI/Private/MetalRHI.cpp
@@ -56,7 +56,7 @@ FMetalDynamicRHI::FMetalDynamicRHI()
 //	GRHIVendorId = 
 	GSupportsRenderTargetFormat_PF_G8 = false;
 	GSupportsQuads = false;
-	GRHISupportsTextureStreaming = false;
+	GRHISupportsTextureStreaming = true;
 	GMaxShadowDepthBufferSizeX = 4096;
 	GMaxShadowDepthBufferSizeY = 4096;
 // 	GReadTexturePoolSizeFromIni = true;
diff --git a/Engine/Source/Runtime/IOS/MetalRHI/Private/MetalTexture.cpp b/Engine/Source/Runtime/IOS/MetalRHI/Private/MetalTexture.cpp
--- a/Engine/Source/Runtime/IOS/MetalRHI/Private/MetalTexture.cpp
+++ b/Engine/Source/Runtime/IOS/MetalRHI/Private/MetalTexture.cpp
@@ -78,6 +78,11 @@ FMetalSurface& GetMetalSurfaceFromRHITexture(FRHITexture* Texture)
 	}
 }
 
+static bool IsRenderTarget(uint32 Flags)
+{
+	return (Flags & (TexCreate_RenderTargetable | TexCreate_ResolveTargetable | TexCreate_DepthStencilTargetable)) != 0;
+}
+
 FMetalSurface::FMetalSurface(ERHIResourceType ResourceType, EPixelFormat Format, uint32 InSizeX, uint32 InSizeY, uint32 InSizeZ, uint32 NumSamples, bool bArray, uint32 ArraySize, uint32 NumMips, uint32 InFlags, FResourceBulkDataInterface* BulkData)
 	: PixelFormat(Format)
     , MSAATexture(nil)
@@ -86,13 +91,15 @@ FMetalSurface::FMetalSurface(ERHIResourceType ResourceType, EPixelFormat Format,
 	, SizeZ(InSizeZ)
 	, bIsCubemap(false)
 	, Flags(InFlags)
+	, TotalTextureSize(0)
 {
 	// the special back buffer surface will be updated in FMetalManager::BeginFrame - no need to set the texture here
 	if (Flags & TexCreate_Presentable)
 	{
 		return;
 	}
-	
+
+	bool bIsRenderTarget = IsRenderTarget(Flags);
 	MTLPixelFormat MTLFormat = (MTLPixelFormat)GPixelFormats[Format].PlatformFormat;
 	MTLTextureDescriptor* Desc;
 	
@@ -140,13 +147,16 @@ FMetalSurface::FMetalSurface(ERHIResourceType ResourceType, EPixelFormat Format,
 		BulkData->Discard();
 	}
 	
-	LockedMemory = NULL;
+	FMemory::Memzero(LockedMemory, sizeof(LockedMemory));
+
+	// calculate size of the texture
+	TotalTextureSize = GetMemorySize();
 
 	if (!FParse::Param(FCommandLine::Get(), TEXT("nomsaa")))
 	{
 		if (NumSamples > 1)
 		{
-			check(Flags & (TexCreate_RenderTargetable | TexCreate_DepthStencilTargetable));
+			check(bIsRenderTarget);
 			Desc.textureType = MTLTextureType2DMultisample;
 	
 			// allow commandline to override
@@ -158,6 +168,14 @@ FMetalSurface::FMetalSurface(ERHIResourceType ResourceType, EPixelFormat Format,
 			{
 				[Texture release];
 				Texture = MSAATexture;
+
+				// we don't have the resolve texture, so we just update the memory size with the MSAA size
+				TotalTextureSize = TotalTextureSize * NumSamples;
+			}
+			else
+			{
+				// an MSAA render target takes NumSamples more space, in addition to the resolve texture
+				TotalTextureSize += TotalTextureSize * NumSamples;
 			}
         
 			NSLog(@"Creating %dx MSAA %d x %d %s surface", (int32)Desc.sampleCount, SizeX, SizeY, (Flags & TexCreate_RenderTargetable) ? "Color" : "Depth");
@@ -168,12 +186,68 @@ FMetalSurface::FMetalSurface(ERHIResourceType ResourceType, EPixelFormat Format,
 			TRACK_OBJECT(MSAATexture);
 		}
 	}
+
+	// track memory usage
+	if (bIsRenderTarget)
+	{
+		GCurrentRendertargetMemorySize += Align(TotalTextureSize, 1024) / 1024;
+	}
+	else
+	{
+		GCurrentTextureMemorySize += Align(TotalTextureSize, 1024) / 1024;
+	}
+
+#if STATS
+	if (ResourceType == RRT_TextureCube)
+	{
+		if (bIsRenderTarget)
+		{
+			INC_MEMORY_STAT_BY(STAT_RenderTargetMemoryCube, TotalTextureSize);
+		}
+		else
+		{
+			INC_MEMORY_STAT_BY(STAT_TextureMemoryCube, TotalTextureSize);
+		}
+	}
+	else if (ResourceType == RRT_Texture3D)
+	{
+		if (bIsRenderTarget)
+		{
+			INC_MEMORY_STAT_BY(STAT_RenderTargetMemory3D, TotalTextureSize);
+		}
+		else
+		{
+			INC_MEMORY_STAT_BY(STAT_TextureMemory3D, TotalTextureSize);
+		}
+	}
+	else
+	{
+		if (bIsRenderTarget)
+		{
+			INC_MEMORY_STAT_BY(STAT_RenderTargetMemory2D, TotalTextureSize);
+		}
+		else
+		{
+			INC_MEMORY_STAT_BY(STAT_TextureMemory2D, TotalTextureSize);
+		}
+	}
+#endif
 }
 
 
 FMetalSurface::~FMetalSurface()
 {
-    if (MSAATexture != nil)
+	// track memory usage
+	if (IsRenderTarget(Flags))
+	{
+		GCurrentRendertargetMemorySize -= Align(TotalTextureSize, 1024) / 1024;
+	}
+	else
+	{
+		GCurrentTextureMemorySize -= Align(TotalTextureSize, 1024) / 1024;
+	}
+
+	if (MSAATexture != nil)
     {
 		FMetalManager::ReleaseObject(MSAATexture);
     }
@@ -187,47 +261,27 @@ FMetalSurface::~FMetalSurface()
 
 void* FMetalSurface::Lock(uint32 MipIndex, uint32 ArrayIndex, EResourceLockMode LockMode, uint32& DestStride)
 {
-	// Calculate the dimensions of the mip-map.
-	const uint32 BlockSizeX = GPixelFormats[PixelFormat].BlockSizeX;
-	const uint32 BlockSizeY = GPixelFormats[PixelFormat].BlockSizeY;
-	const uint32 BlockBytes = GPixelFormats[PixelFormat].BlockBytes;
-	const uint32 MipSizeX = FMath::Max(SizeX >> MipIndex,BlockSizeX);
-	const uint32 MipSizeY = FMath::Max(SizeY >> MipIndex,BlockSizeY);
-	uint32 NumBlocksX = (MipSizeX + BlockSizeX - 1) / BlockSizeX;
-	uint32 NumBlocksY = (MipSizeY + BlockSizeY - 1) / BlockSizeY;
-	if ( PixelFormat == PF_PVRTC2 || PixelFormat == PF_PVRTC4 )
-	{
-		// PVRTC has minimum 2 blocks width and height
-		NumBlocksX = FMath::Max<uint32>(NumBlocksX, 2);
-		NumBlocksY = FMath::Max<uint32>(NumBlocksY, 2);
-	}
-	const uint32 MipBytes = NumBlocksX * NumBlocksY * BlockBytes;
-	DestStride = NumBlocksX * BlockBytes;
+	// get size and stride
+	const uint32 MipBytes = GetMipSize(MipIndex, &DestStride);
 	
 	// allocate some temporary memory
-	check(LockedMemory == NULL);
-	LockedMemory = FMemory::Malloc(MipBytes);
+	check(LockedMemory[MipIndex] == NULL);
+	LockedMemory[MipIndex] = FMemory::Malloc(MipBytes);
 	if (LockMode != RLM_WriteOnly)
 	{
 		// [Texture readPixels ...];
 	}
 
-	return LockedMemory;
+	return LockedMemory[MipIndex];
 }
 
 void FMetalSurface::Unlock(uint32 MipIndex, uint32 ArrayIndex)
 {
-	const uint32 BlockSizeX = GPixelFormats[PixelFormat].BlockSizeX;
-	const uint32 BlockSizeY = GPixelFormats[PixelFormat].BlockSizeY;
-	const uint32 BlockBytes = GPixelFormats[PixelFormat].BlockBytes;
-	const uint32 MipSizeX = FMath::Max(SizeX >> MipIndex,BlockSizeX);
-	const uint32 MipSizeY = FMath::Max(SizeY >> MipIndex,BlockSizeY);
-	uint32 NumBlocksX = (MipSizeX + BlockSizeX - 1) / BlockSizeX;
-	uint32 NumBlocksY = (MipSizeY + BlockSizeY - 1) / BlockSizeY;
-	uint32 MipBytes = NumBlocksX * NumBlocksY * BlockBytes;
-	uint32 Stride = NumBlocksX * BlockBytes;
+	uint32 Stride;
+	uint32 MipBytes = GetMipSize(MipIndex, &Stride);
 	if (PixelFormat == PF_PVRTC2 || PixelFormat == PF_PVRTC4)
 	{
+		// compressed textures want zero here for whatever reason
 		MipBytes = 0;
 		Stride = 0;
 	}
@@ -236,15 +290,60 @@ void FMetalSurface::Unlock(uint32 MipIndex, uint32 ArrayIndex)
 
 	// upload the texture to the texture slice
 	MTLRegion Region = MTLRegionMake2D(0, 0, FMath::Max<uint32>(SizeX>>MipIndex, 1), FMath::Max<uint32>(SizeY>>MipIndex, 1));
-	[Texture replaceRegion:Region mipmapLevel:MipIndex slice:ArrayIndex withBytes:LockedMemory bytesPerRow:Stride bytesPerImage:MipBytes];
+	[Texture replaceRegion:Region mipmapLevel:MipIndex slice:ArrayIndex withBytes:LockedMemory[MipIndex] bytesPerRow:Stride bytesPerImage:MipBytes];
 	
-	FMemory::Free(LockedMemory);
-	LockedMemory = NULL;
+	FMemory::Free(LockedMemory[MipIndex]);
+	LockedMemory[MipIndex] = NULL;
+}
+
+uint32 FMetalSurface::GetMipSize(uint32 MipIndex, uint32* Stride)
+{
+	// Calculate the dimensions of the mip-map.
+	const uint32 BlockSizeX = GPixelFormats[PixelFormat].BlockSizeX;
+	const uint32 BlockSizeY = GPixelFormats[PixelFormat].BlockSizeY;
+	const uint32 BlockBytes = GPixelFormats[PixelFormat].BlockBytes;
+	const uint32 MipSizeX = FMath::Max(SizeX >> MipIndex, BlockSizeX);
+	const uint32 MipSizeY = FMath::Max(SizeY >> MipIndex, BlockSizeY);
+	const uint32 MipSizeZ = FMath::Max(SizeZ >> MipIndex, 1u);
+	uint32 NumBlocksX = (MipSizeX + BlockSizeX - 1) / BlockSizeX;
+	uint32 NumBlocksY = (MipSizeY + BlockSizeY - 1) / BlockSizeY;
+	if (PixelFormat == PF_PVRTC2 || PixelFormat == PF_PVRTC4)
+	{
+		// PVRTC has minimum 2 blocks width and height
+		NumBlocksX = FMath::Max<uint32>(NumBlocksX, 2);
+		NumBlocksY = FMath::Max<uint32>(NumBlocksY, 2);
+	}
+
+	const uint32 MipBytes = NumBlocksX * NumBlocksY * BlockBytes * MipSizeZ;
+
+	if (Stride)
+	{
+		*Stride = NumBlocksX * BlockBytes;
+	}
+
+	return MipBytes;
 }
 
 uint32 FMetalSurface::GetMemorySize()
 {
-	return 0;
+	// if already calculated, no need to do it again
+	if (TotalTextureSize != 0)
+	{
+		return TotalTextureSize;
+	}
+
+	if (Texture == nil)
+	{
+		return 0;
+	}
+
+	uint32 TotalSize = 0;
+	for (uint32 MipIndex = 0; MipIndex < [Texture mipmapLevelCount]; MipIndex++)
+	{
+		TotalSize += GetMipSize(MipIndex, NULL);
+	}
+
+	return TotalSize;
 }
 
 /*-----------------------------------------------------------------------------
@@ -253,7 +352,14 @@ uint32 FMetalSurface::GetMemorySize()
 
 void FMetalDynamicRHI::RHIGetTextureMemoryStats(FTextureMemoryStats& OutStats)
 {
+	OutStats.DedicatedVideoMemory = 0;
+	OutStats.DedicatedSystemMemory = 0;
+	OutStats.SharedSystemMemory = 0;
+	OutStats.TotalGraphicsMemory = 0;
 
+	OutStats.AllocatedMemorySize = int64(GCurrentTextureMemorySize) * 1024;
+	OutStats.TexturePoolSize = GTexturePoolSize;
+	OutStats.PendingMemoryAdjustment = 0;
 }
 
 bool FMetalDynamicRHI::RHIGetTextureMemoryVisualizeData( FColor* /*TextureData*/, int32 /*SizeX*/, int32 /*SizeY*/, int32 /*Pitch*/, int32 /*PixelSize*/ )
@@ -314,12 +420,57 @@ FTexture2DRHIRef FMetalDynamicRHI::RHIAsyncReallocateTexture2D(FTexture2DRHIPara
 {
 	DYNAMIC_CAST_METGALRESOURCE(Texture2D,OldTexture);
 
-	return NULL;
+	FMetalTexture2D* NewTexture = new FMetalTexture2D(OldTexture->GetFormat(), NewSizeX, NewSizeY, NewMipCount, OldTexture->GetNumSamples(), OldTexture->GetFlags(), NULL);
+
+	// @todo: gather these all up over a frame
+	id<MTLCommandBuffer> CommandBuffer = FMetalManager::Get()->CreateTempCommandBuffer(false/*bRetainReferences*/);
+
+	// create a blitter object
+	id<MTLBlitCommandEncoder> Blitter = [CommandBuffer blitCommandEncoder];
+
+
+	// figure out what mips to schedule
+	const uint32 NumSharedMips = FMath::Min(OldTexture->GetNumMips(), NewTexture->GetNumMips());
+	const uint32 SourceMipOffset = OldTexture->GetNumMips() - NumSharedMips;
+	const uint32 DestMipOffset = NewTexture->GetNumMips() - NumSharedMips;
+
+	// only handling straight 2D textures here
+	uint32 SliceIndex = 0;
+	MTLOrigin Origin = MTLOriginMake(0,0,0);
+
+	for (uint32 MipIndex = 0; MipIndex < NumSharedMips; ++MipIndex)
+	{
+		const uint32 MipSizeX = FMath::Max<uint32>(1, NewSizeX >> (MipIndex + DestMipOffset));
+		const uint32 MipSizeY = FMath::Max<uint32>(1, NewSizeY >> (MipIndex + DestMipOffset));
+
+		// set up the copy
+		[Blitter copyFromTexture:OldTexture->Surface.Texture 
+					 sourceSlice:SliceIndex
+					 sourceLevel:MipIndex + SourceMipOffset
+					sourceOrigin:Origin
+					  sourceSize:MTLSizeMake(MipSizeX, MipSizeY, 0)
+					   toTexture:NewTexture->Surface.Texture
+			    destinationSlice:SliceIndex
+			    destinationLevel:MipIndex + DestMipOffset
+			   destinationOrigin:Origin];
+	}
+
+	// when done, decrement the counter to indicate it's safe
+	[CommandBuffer addCompletedHandler:^(id <MTLCommandBuffer> Buffer)
+	{
+		RequestStatus->Decrement();
+	}];
+
+	// kick it off!
+	[Blitter endEncoding];
+	[CommandBuffer commit];
+
+	return NewTexture;
 }
 
 ETextureReallocationStatus FMetalDynamicRHI::RHIFinalizeAsyncReallocateTexture2D( FTexture2DRHIParamRef Texture2D, bool bBlockUntilCompleted )
 {
-	return TexRealloc_Failed;
+	return TexRealloc_Succeeded;
 }
 
 ETextureReallocationStatus FMetalDynamicRHI::RHICancelAsyncReallocateTexture2D( FTexture2DRHIParamRef Texture2D, bool bBlockUntilCompleted )
diff --git a/Engine/Source/Runtime/IOS/MetalRHI/Public/MetalResources.h b/Engine/Source/Runtime/IOS/MetalRHI/Public/MetalResources.h
--- a/Engine/Source/Runtime/IOS/MetalRHI/Public/MetalResources.h
+++ b/Engine/Source/Runtime/IOS/MetalRHI/Public/MetalResources.h
@@ -158,11 +158,18 @@ public:
 	 */
 	void Unlock(uint32 MipIndex, uint32 ArrayIndex);
 
+	/**
+	 * Returns how much memory a single mip uses, and optionally returns the stride
+	 */
+	uint32 GetMipSize(uint32 MipIndex, uint32* Stride);
+
 	/**
 	 * Returns how much memory is used by the surface
 	 */
 	uint32 GetMemorySize();
 
+
+
 	EPixelFormat PixelFormat;
 	id<MTLTexture> Texture;
     id<MTLTexture> MSAATexture;
@@ -170,7 +177,11 @@ public:
 	bool bIsCubemap;
 	
 	uint32 Flags;
-	void* LockedMemory;
+	// one per mip
+	void* LockedMemory[16];
+
+	// how much memory is allocated for this texture
+	uint64 TotalTextureSize;
 };
 
 class FMetalTexture2D : public FRHITexture2D
