- Fixed Slate text rendering in Metal

diff --git a/Engine/Source/Runtime/IOS/MetalRHI/Private/MetalRHI.cpp b/Engine/Source/Runtime/IOS/MetalRHI/Private/MetalRHI.cpp
--- a/Engine/Source/Runtime/IOS/MetalRHI/Private/MetalRHI.cpp
+++ b/Engine/Source/Runtime/IOS/MetalRHI/Private/MetalRHI.cpp
@@ -114,7 +114,7 @@ FMetalDynamicRHI::FMetalDynamicRHI()
 	GPixelFormats[PF_A1					].PlatformFormat	= MTLPixelFormatInvalid;
 	GPixelFormats[PF_FloatR11G11B10		].PlatformFormat	= MTLPixelFormatRG11B10Float;
 	GPixelFormats[PF_FloatR11G11B10		].BlockBytes		= 4;
-	GPixelFormats[PF_A8					].PlatformFormat	= MTLPixelFormatR8Unorm;
+	GPixelFormats[PF_A8					].PlatformFormat	= MTLPixelFormatA8Unorm;
 	GPixelFormats[PF_R32_UINT			].PlatformFormat	= MTLPixelFormatR32Uint;
 	GPixelFormats[PF_R32_SINT			].PlatformFormat	= MTLPixelFormatR32Sint;
 	GPixelFormats[PF_R16G16B16A16_UINT	].PlatformFormat	= MTLPixelFormatRGBA16Uint;
