Fix for a previous IOS compile error related to duplicate symbols found - removed the IOS specific copy and will let the linker sort things out.

diff --git a/Engine/Source/Runtime/IOS/IOSAudio/Private/IOSAudioSource.cpp b/Engine/Source/Runtime/IOS/IOSAudio/Private/IOSAudioSource.cpp
--- a/Engine/Source/Runtime/IOS/IOSAudio/Private/IOSAudioSource.cpp
+++ b/Engine/Source/Runtime/IOS/IOSAudio/Private/IOSAudioSource.cpp
@@ -11,213 +11,14 @@
 #include "IOSAudioDevice.h"
 #include "Engine.h"
 
-#define NUM_ADAPTATION_TABLE 16
-#define NUM_ADAPTATION_COEFF 7
 #define SOUND_SOURCE_FREE 0
 #define SOUND_SOURCE_LOCKED 1
 
-namespace
-{
-	template <typename T, uint32 B>
-	inline T SignExtend(const T ValueToExtend)
-	{
-		struct { T ExtendedValue:B; } SignExtender;
-		return SignExtender.ExtendedValue = ValueToExtend;
-	}
-
-	template <typename T>
-	inline T ReadFromByteStream(const uint8* ByteStream, int32& ReadIndex, bool bLittleEndian = true)
-	{
-		T ValueRaw = 0;
-
-		if (bLittleEndian)
-		{
-#if PLATFORM_LITTLE_ENDIAN
-			for (int32 ByteIndex = 0; ByteIndex < sizeof(T); ++ByteIndex)
-#else
-			for (int32 ByteIndex = sizeof(T) - 1; ByteIndex >= 0; --ByteIndex)
-#endif // PLATFORM_LITTLE_ENDIAN
-			{
-				ValueRaw |= ByteStream[ReadIndex++] << 8 * ByteIndex;
-			}
-		}
-		else
-		{
-#if PLATFORM_LITTLE_ENDIAN
-			for (int32 ByteIndex = sizeof(T) - 1; ByteIndex >= 0; --ByteIndex)
-#else
-			for (int32 ByteIndex = 0; ByteIndex < sizeof(T); ++ByteIndex)
-#endif // PLATFORM_LITTLE_ENDIAN
-			{
-				ValueRaw |= ByteStream[ReadIndex++] << 8 * ByteIndex;
-			}
-		}
-
-		return ValueRaw;
-	}
-
-	template <typename T>
-	inline void WriteToByteStream(T Value, uint8* ByteStream, int32& WriteIndex, bool bLittleEndian = true)
-	{
-		if (bLittleEndian)
-		{
-#if PLATFORM_LITTLE_ENDIAN
-			for (int32 ByteIndex = 0; ByteIndex < sizeof(T); ++ByteIndex)
-#else
-			for (int32 ByteIndex = sizeof(T) - 1; ByteIndex >= 0; --ByteIndex)
-#endif // PLATFORM_LITTLE_ENDIAN
-			{
-				ByteStream[WriteIndex++] = (Value >> (8 * ByteIndex)) & 0xFF;
-			}
-		}
-		else
-		{
-#if PLATFORM_LITTLE_ENDIAN
-			for (int32 ByteIndex = sizeof(T) - 1; ByteIndex >= 0; --ByteIndex)
-#else
-			for (int32 ByteIndex = 0; ByteIndex < sizeof(T); ++ByteIndex)
-#endif // PLATFORM_LITTLE_ENDIAN
-			{
-				ByteStream[WriteIndex++] = (Value >> (8 * ByteIndex)) & 0xFF;
-			}
-		}
-	}
-
-	template <typename T>
-	inline T ReadFromArray(const T* ElementArray, int32& ReadIndex, int32 NumElements, int32 IndexStride = 1)
-	{
-		T OutputValue = 0;
-
-		if (ReadIndex >= 0 && ReadIndex < NumElements)
-		{
-			OutputValue = ElementArray[ReadIndex];
-			ReadIndex += IndexStride;
-		}
-
-		return OutputValue;
-	}
-	
-	inline bool LockSourceChannel(int32* ChannelLock)
-	{
-		check(ChannelLock != NULL);
-		return FPlatformAtomics::InterlockedCompareExchange(ChannelLock, SOUND_SOURCE_LOCKED, SOUND_SOURCE_FREE) == SOUND_SOURCE_FREE;
-	}
-	
-	inline void UnlockSourceChannel(int32* ChannelLock)
-	{
-		check(ChannelLock != NULL);
-		int32 Result = FPlatformAtomics::InterlockedCompareExchange(ChannelLock, SOUND_SOURCE_FREE, SOUND_SOURCE_LOCKED);
-		
-		check(Result == SOUND_SOURCE_LOCKED);
-	}
-	
-} // end namespace
-
+// Linker will figure this out for us on IOS as it is now in a core source file
 namespace ADPCM
 {
-	template <typename T>
-	static void GetAdaptationTable(T(& OutAdaptationTable)[NUM_ADAPTATION_TABLE])
-	{
-		// Magic values as specified by standard
-		static T AdaptationTable[] =
-		{
-			230, 230, 230, 230, 307, 409, 512, 614,
-			768, 614, 512, 409, 307, 230, 230, 230
-		};
-
-		FMemory::Memcpy(&OutAdaptationTable, AdaptationTable, sizeof(AdaptationTable));
-	}
-
-	template <typename T>
-	static void GetAdaptationCoefficients(T(& OutAdaptationCoefficient1)[NUM_ADAPTATION_COEFF], T(& OutAdaptationCoefficient2)[NUM_ADAPTATION_COEFF])
-	{
-		// Magic values as specified by standard
-		static T AdaptationCoefficient1[] =
-		{
-			256, 512, 0, 192, 240, 460, 392
-		};
-		static T AdaptationCoefficient2[] =
-		{
-			0, -256, 0,  64, 0, -208, -232
-		};
-
-		FMemory::Memcpy(&OutAdaptationCoefficient1, AdaptationCoefficient1, sizeof(AdaptationCoefficient1));
-		FMemory::Memcpy(&OutAdaptationCoefficient2, AdaptationCoefficient2, sizeof(AdaptationCoefficient2));
-	}
-
-	struct FAdaptationContext
-	{
-	public:
-		// Adaptation constants
-		int32 AdaptationTable[NUM_ADAPTATION_TABLE];
-		int32 AdaptationCoefficient1[NUM_ADAPTATION_COEFF];
-		int32 AdaptationCoefficient2[NUM_ADAPTATION_COEFF];
-
-		int32 AdaptationDelta;
-		int32 Coefficient1;
-		int32 Coefficient2;
-		int32 Sample1;
-		int32 Sample2;
-
-		FAdaptationContext() :
-		AdaptationDelta(0),
-			Coefficient1(0),
-			Coefficient2(0),
-			Sample1(0),
-			Sample2(0)
-		{
-			GetAdaptationTable(AdaptationTable);
-			GetAdaptationCoefficients(AdaptationCoefficient1, AdaptationCoefficient2);
-		}
-	};
-
-	int16 DecodeNibble(FAdaptationContext& Context, uint8 EncodedNibble)
-	{
-		int32 PredictedSample = (Context.Sample1 * Context.Coefficient1 + Context.Sample2 * Context.Coefficient2) / 256;
-		PredictedSample += SignExtend<int8, 4>(EncodedNibble) * Context.AdaptationDelta;
-		PredictedSample = FMath::Clamp(PredictedSample, -32768, 32767);
-
-		// Shuffle samples for the next iteration
-		Context.Sample2 = Context.Sample1;
-		Context.Sample1 = static_cast<int16>(PredictedSample);
-		Context.AdaptationDelta = (Context.AdaptationDelta * Context.AdaptationTable[EncodedNibble]) / 256;
-		Context.AdaptationDelta = FMath::Max(Context.AdaptationDelta, 16);
-
-		return Context.Sample1;
-	}
-
-	void DecodeBlock(const uint8* EncodedADPCMBlock, int32 BlockSize, int16* DecodedPCMData)
-	{
-		FAdaptationContext Context;
-		int32 ReadIndex = 0;
-		int32 WriteIndex = 0;
-
-		uint8 CoefficientIndex = ReadFromByteStream<uint8>(EncodedADPCMBlock, ReadIndex);
-		Context.AdaptationDelta = ReadFromByteStream<int16>(EncodedADPCMBlock, ReadIndex);
-		Context.Sample1 = ReadFromByteStream<int16>(EncodedADPCMBlock, ReadIndex);
-		Context.Sample2 = ReadFromByteStream<int16>(EncodedADPCMBlock, ReadIndex);
-		Context.Coefficient1 = Context.AdaptationCoefficient1[CoefficientIndex];
-		Context.Coefficient2 = Context.AdaptationCoefficient2[CoefficientIndex];
-
-		// The first two samples are sent directly to the output in reverse order, as per the standard
-		DecodedPCMData[WriteIndex++] = Context.Sample2;
-		DecodedPCMData[WriteIndex++] = Context.Sample1;
-
-		uint8 EncodedNibblePair = 0;
-		uint8 EncodedNibble = 0;
-		while (ReadIndex < BlockSize)
-		{
-			// Read from the byte stream and advance the read head.
-			EncodedNibblePair = ReadFromByteStream<uint8>(EncodedADPCMBlock, ReadIndex);
-
-			EncodedNibble = (EncodedNibblePair >> 4) & 0x0F;
-			DecodedPCMData[WriteIndex++] = DecodeNibble(Context, EncodedNibble);
-
-			EncodedNibble = EncodedNibblePair & 0x0F;
-			DecodedPCMData[WriteIndex++] = DecodeNibble(Context, EncodedNibble);
-		}
-	}
-} // end namespace ADPCM
+	void DecodeBlock(const uint8* EncodedADPCMBlock, int32 BlockSize, int16* DecodedPCMData);
+}
 
 /*------------------------------------------------------------------------------------
 	FIOSAudioSoundSource
