Checking in fix 2 - got slightly over zealous with function removal in the anonymous namespace

diff --git a/Engine/Source/Runtime/IOS/IOSAudio/Private/IOSAudioSource.cpp b/Engine/Source/Runtime/IOS/IOSAudio/Private/IOSAudioSource.cpp
--- a/Engine/Source/Runtime/IOS/IOSAudio/Private/IOSAudioSource.cpp
+++ b/Engine/Source/Runtime/IOS/IOSAudio/Private/IOSAudioSource.cpp
@@ -14,6 +14,24 @@
 #define SOUND_SOURCE_FREE 0
 #define SOUND_SOURCE_LOCKED 1
 
+namespace
+{
+	inline bool LockSourceChannel(int32* ChannelLock)
+	{
+		check(ChannelLock != NULL);
+		return FPlatformAtomics::InterlockedCompareExchange(ChannelLock, SOUND_SOURCE_LOCKED, SOUND_SOURCE_FREE) == SOUND_SOURCE_FREE;
+	}
+
+	inline void UnlockSourceChannel(int32* ChannelLock)
+	{
+		check(ChannelLock != NULL);
+		int32 Result = FPlatformAtomics::InterlockedCompareExchange(ChannelLock, SOUND_SOURCE_FREE, SOUND_SOURCE_LOCKED);
+
+		check(Result == SOUND_SOURCE_LOCKED);
+	}
+
+} // end namespace
+
 // Linker will figure this out for us on IOS as it is now in a core source file
 namespace ADPCM
 {
