Add IOS support to listen if other BG sounds are playing.

diff --git a/Engine/Source/Runtime/Engine/Public/AudioDevice.h b/Engine/Source/Runtime/Engine/Public/AudioDevice.h
--- a/Engine/Source/Runtime/Engine/Public/AudioDevice.h
+++ b/Engine/Source/Runtime/Engine/Public/AudioDevice.h
@@ -508,7 +508,10 @@ public:
     
     /** Resume any context related objects */
     virtual void ResumeContext() {}
-    
+
+	/** Check if any background music or sound is playing through the audio device */
+	virtual bool IsExernalBackgroundSoundActive() { return false; }
+
 protected:
 	friend class FSoundSource;
 
diff --git a/Engine/Source/Runtime/IOS/IOSAudio/Private/IOSAudioSession.cpp b/Engine/Source/Runtime/IOS/IOSAudio/Private/IOSAudioSession.cpp
--- a/Engine/Source/Runtime/IOS/IOSAudio/Private/IOSAudioSession.cpp
+++ b/Engine/Source/Runtime/IOS/IOSAudio/Private/IOSAudioSession.cpp
@@ -33,3 +33,8 @@ bool FIOSAudioDevice::SetAudioSessionActive(bool bActive)
 	AVAudioSession* AudioSession = [AVAudioSession sharedInstance];
 	return [AudioSession setActive:bActive error:nil] == YES;
 }
+
+bool FIOSAudioDevice::IsExernalBackgroundSoundActive()
+{
+	return [[AVAudioSession sharedInstance] isOtherAudioPlaying];
+}
\ No newline at end of file
diff --git a/Engine/Source/Runtime/IOS/IOSAudio/Public/IOSAudioDevice.h b/Engine/Source/Runtime/IOS/IOSAudio/Public/IOSAudioDevice.h
--- a/Engine/Source/Runtime/IOS/IOSAudio/Public/IOSAudioDevice.h
+++ b/Engine/Source/Runtime/IOS/IOSAudio/Public/IOSAudioDevice.h
@@ -216,6 +216,9 @@ protected:
 	void GetHardwareSampleRate(double& OutSampleRate);
 	bool SetHardwareSampleRate(const double& InSampleRate);
 	bool SetAudioSessionActive(bool bActive);
+
+	/** Check if any background music or sound is playing through the audio device */
+	virtual bool IsExernalBackgroundSoundActive() override;
 	
 private:
 	void HandleError(const TCHAR* InLogOutput, bool bTeardown = false);
