UE4 - Metal integration
- Adds resolve actions
- CPU Optimization for setting uniform buffers
- Fixes texture async realloc assert

diff --git a/Engine/Source/Runtime/IOS/MetalRHI/Private/MetalCommands.cpp b/Engine/Source/Runtime/IOS/MetalRHI/Private/MetalCommands.cpp
--- a/Engine/Source/Runtime/IOS/MetalRHI/Private/MetalCommands.cpp
+++ b/Engine/Source/Runtime/IOS/MetalRHI/Private/MetalCommands.cpp
@@ -252,6 +252,14 @@ void FMetalDynamicRHI::RHISetShaderUniformBuffer(FVertexShaderRHIParamRef Vertex
 	DYNAMIC_CAST_METGALRESOURCE(VertexShader, VertexShader);
 	VertexShader->BoundUniformBuffers[BufferIndex] = BufferRHI;
 	VertexShader->DirtyUniformBuffers |= 1 << BufferIndex;
+
+	auto& Bindings = VertexShader->Bindings;
+	check(BufferIndex < Bindings.NumUniformBuffers);
+	if (Bindings.bHasRegularUniformBuffers)
+	{
+		auto* UB = (FMetalUniformBuffer*)BufferRHI;
+		[FMetalManager::GetContext() setVertexBuffer:UB->Buffer offset:UB->Offset atIndex:BufferIndex];
+	}
 }
 
 void FMetalDynamicRHI::RHISetShaderUniformBuffer(FHullShaderRHIParamRef HullShader, uint32 BufferIndex, FUniformBufferRHIParamRef BufferRHI)
@@ -274,6 +282,14 @@ void FMetalDynamicRHI::RHISetShaderUniformBuffer(FPixelShaderRHIParamRef PixelSh
 	DYNAMIC_CAST_METGALRESOURCE(PixelShader, PixelShader);
 	PixelShader->BoundUniformBuffers[BufferIndex] = BufferRHI;
 	PixelShader->DirtyUniformBuffers |= 1 << BufferIndex;
+
+	auto& Bindings = PixelShader->Bindings;
+	check(BufferIndex < Bindings.NumUniformBuffers);
+	if (Bindings.bHasRegularUniformBuffers)
+	{
+		auto* UB = (FMetalUniformBuffer*)BufferRHI;
+		[FMetalManager::GetContext() setFragmentBuffer:UB->Buffer offset:UB->Offset atIndex:BufferIndex];
+	}
 }
 
 void FMetalDynamicRHI::RHISetShaderUniformBuffer(FComputeShaderRHIParamRef ComputeShader, uint32 BufferIndex, FUniformBufferRHIParamRef BufferRHI)
@@ -305,15 +321,16 @@ void FMetalDynamicRHI::RHISetRenderTargets(uint32 NumSimultaneousRenderTargets,
 
 	for (int32 RenderTargetIndex = 0; RenderTargetIndex < MaxMetalRenderTargets; RenderTargetIndex++)
 	{
+		const FRHIRenderTargetView& RenderTargetView = NewRenderTargets[RenderTargetIndex];
 		// update the current RTs
-		if (RenderTargetIndex < NumSimultaneousRenderTargets && NewRenderTargets[RenderTargetIndex].Texture != NULL)
+		if (RenderTargetIndex < NumSimultaneousRenderTargets && RenderTargetView.Texture != NULL)
 		{
-			FMetalSurface& Surface = GetMetalSurfaceFromRHITexture(NewRenderTargets[0].Texture);
-			Manager->SetCurrentRenderTarget(&Surface, RenderTargetIndex, NumSimultaneousRenderTargets);
+			FMetalSurface& Surface = GetMetalSurfaceFromRHITexture(RenderTargetView.Texture);
+			Manager->SetCurrentRenderTarget(&Surface, RenderTargetIndex, RenderTargetView.MipIndex, RenderTargetView.ArraySliceIndex, GetMetalRTLoadAction(RenderTargetView.LoadAction), GetMetalRTStoreAction(RenderTargetView.StoreAction), NumSimultaneousRenderTargets);
 		}
 		else
 		{
-			Manager->SetCurrentRenderTarget(NULL, RenderTargetIndex, NumSimultaneousRenderTargets);
+			Manager->SetCurrentRenderTarget(NULL, RenderTargetIndex, 0, 0, MTLLoadActionDontCare, MTLStoreActionStore, NumSimultaneousRenderTargets);
 		}
 	}
 	
@@ -333,8 +350,13 @@ void FMetalDynamicRHI::RHISetRenderTargets(uint32 NumSimultaneousRenderTargets,
 	// Set the viewport to the full size of render target 0.
 	if (NewRenderTargets && NewRenderTargets[0].Texture)
 	{
-		FMetalSurface& RenderTarget = GetMetalSurfaceFromRHITexture(NewRenderTargets[0].Texture);
-		RHISetViewport(0, 0, 0.0f, RenderTarget.Texture.width, RenderTarget.Texture.height, 1.0f);
+		const FRHIRenderTargetView& RenderTargetView = NewRenderTargets[0];
+		FMetalSurface& RenderTarget = GetMetalSurfaceFromRHITexture(RenderTargetView.Texture);
+		
+		uint32 Width = FMath::Max(RenderTarget.Texture.width >> RenderTargetView.MipIndex, (uint32)1);
+		uint32 Height = FMath::Max(RenderTarget.Texture.height >> RenderTargetView.MipIndex, (uint32)1);
+		
+		RHISetViewport(0, 0, 0.0f, Width, Height, 1.0f);
 	}
 }
 
@@ -364,6 +386,8 @@ void FMetalDynamicRHI::RHIDrawPrimitive(uint32 PrimitiveType, uint32 BaseVertexI
 	SCOPE_CYCLE_COUNTER(STAT_MetalDrawCallTime);
 	checkf(NumInstances == 1, TEXT("Currently only 1 instance is supported"));
 	
+	RHI_DRAW_CALL_STATS(PrimitiveType,NumInstances*NumPrimitives);
+
 	// how many verts to render
 	uint32 NumVertices = GetVertexCountForPrimitiveCount(NumPrimitives, PrimitiveType);
 
@@ -390,6 +414,8 @@ void FMetalDynamicRHI::RHIDrawIndexedPrimitive(FIndexBufferRHIParamRef IndexBuff
 	checkf(NumInstances == 1, TEXT("Currently only 1 instance is supported"));
 	checkf(BaseVertexIndex  == 0, TEXT("BaseVertexIndex must be 0, see GRHISupportsBaseVertexIndex"));
 
+	RHI_DRAW_CALL_STATS(PrimitiveType,NumInstances*NumPrimitives);
+
 	DYNAMIC_CAST_METGALRESOURCE(IndexBuffer,IndexBuffer);
 
 	// finalize any pending state
@@ -447,6 +473,8 @@ void FMetalDynamicRHI::RHIEndDrawPrimitiveUP()
 {
 	SCOPE_CYCLE_COUNTER(STAT_MetalDrawCallTime);
 
+	RHI_DRAW_CALL_STATS(GPendingPrimitiveType,GPendingNumPrimitives);
+
 	// set the vertex buffer
 	[FMetalManager::GetContext() setVertexBuffer:FMetalManager::Get()->GetRingBuffer() offset:GPendingVertexBufferOffset atIndex:UNREAL_TO_METAL_BUFFER_INDEX(0)];
 	
@@ -490,6 +518,8 @@ void FMetalDynamicRHI::RHIEndDrawIndexedPrimitiveUP()
 {
 	SCOPE_CYCLE_COUNTER(STAT_MetalDrawCallTime);
 
+	RHI_DRAW_CALL_STATS(GPendingPrimitiveType,GPendingNumPrimitives);
+
 	// set the vertex buffer
 	[FMetalManager::GetContext() setVertexBuffer:FMetalManager::Get()->GetRingBuffer() offset:GPendingVertexBufferOffset atIndex:UNREAL_TO_METAL_BUFFER_INDEX(0)];
 
diff --git a/Engine/Source/Runtime/IOS/MetalRHI/Private/MetalManager.cpp b/Engine/Source/Runtime/IOS/MetalRHI/Private/MetalManager.cpp
--- a/Engine/Source/Runtime/IOS/MetalRHI/Private/MetalManager.cpp
+++ b/Engine/Source/Runtime/IOS/MetalRHI/Private/MetalManager.cpp
@@ -38,7 +38,7 @@ static uint32 RTBitOffsets[] = { OFFSET_RENDER_TARGET_FORMAT0, OFFSET_RENDER_TAR
 
 
 
-	
+
 /*
 void FPipelineShadow::SetHash(uint64 InHash)
 {
@@ -194,6 +194,7 @@ FMetalManager::FMetalManager()
 	, PreviousNumRenderTargets(0)
 	, CurrentDepthRenderTexture(nil)
 	, PreviousDepthRenderTexture(nil)
+	, CurrentMSAARenderTexture(nil)
 	, RingBuffer(Device, RingBufferSize, BufferOffsetAlignment)
 	, QueryBuffer(Device, 64 * 1024, 8)
 	, WhichFreeList(0)
@@ -207,6 +208,8 @@ FMetalManager::FMetalManager()
 		CurrentColorRenderTextures[Index] = nil;
 		PreviousColorRenderTextures[Index] = nil;
 	}
+	FMemory::MemSet(CurrentRenderTargetsViewInfo, 0);
+	FMemory::MemSet(PreviousRenderTargetsViewInfo, 0);
 
 	CommandQueue = [Device newCommandQueue];
 
@@ -280,20 +283,10 @@ void FMetalManager::InitFrame()
 	CreateAutoreleasePool();
 
 	// create the command buffer for this frame
-	CurrentCommandBuffer = [CommandQueue commandBufferWithUnretainedReferences];
-	[CurrentCommandBuffer retain];
-	TRACK_OBJECT(CurrentCommandBuffer);
-	
-	uint64 LocalCommandBufferIndex = CommandBufferIndex++;
-	[CurrentCommandBuffer addScheduledHandler : ^ (id <MTLCommandBuffer> Buffer)
-	{
-		FMetalManager::Get()->SetCompletedCommandBufferIndex(LocalCommandBufferIndex);
-	}];
+	CreateCurrentCommandBuffer(true);
 
-//	double Start = FPlatformTime::Seconds();
+	//	double Start = FPlatformTime::Seconds();
 	// block on the semaphore
-	dispatch_semaphore_wait(CommandBufferSemaphore, DISPATCH_TIME_FOREVER);
-
 //	double Mid = FPlatformTime::Seconds();
 
 	// mark us to get later
@@ -309,6 +302,51 @@ void FMetalManager::InitFrame()
 	NumDrawCalls = 0;
 }
 
+void FMetalManager::CreateCurrentCommandBuffer(bool bWait)
+{
+	if (bWait)
+	{
+		dispatch_semaphore_wait(CommandBufferSemaphore, DISPATCH_TIME_FOREVER);
+	}
+	
+	CurrentCommandBuffer = [CommandQueue commandBufferWithUnretainedReferences];
+	[CurrentCommandBuffer retain];
+	TRACK_OBJECT(CurrentCommandBuffer);
+
+	uint64 LocalCommandBufferIndex = CommandBufferIndex++;
+	[CurrentCommandBuffer addScheduledHandler : ^ (id <MTLCommandBuffer> Buffer)
+	{
+		FMetalManager::Get()->SetCompletedCommandBufferIndex(LocalCommandBufferIndex);
+	}];
+}
+
+void FMetalManager::SubmitCommandBufferAndWait()
+{
+	[CurrentCommandBuffer addCompletedHandler : ^ (id <MTLCommandBuffer> Buffer)
+	 {
+		dispatch_semaphore_signal(CommandBufferSemaphore);
+	 }];
+	
+	// commit the render context to the commandBuffer
+	[CurrentContext endEncoding];
+	[CurrentContext release];
+	CurrentContext = nil;
+
+	// kick the whole buffer
+	// Commit to hand the commandbuffer off to the gpu
+	[CurrentCommandBuffer commit];
+
+	// wait for the gpu to finish executing our commands.
+	[CurrentCommandBuffer waitUntilCompleted];
+	
+	//once a commandbuffer is commited it can't be added to again.
+	UNTRACK_OBJECT(CurrentCommandBuffer);
+	[CurrentCommandBuffer release];
+	
+	// create a new command buffer.
+	CreateCurrentCommandBuffer(true);
+}
+
 void FMetalManager::EndFrame(bool bPresent)
 {
 //	NSLog(@"There were %d draw calls for final RT in frame %lld", NumDrawCalls, GFrameCounter);
@@ -318,13 +356,13 @@ void FMetalManager::EndFrame(bool bPresent)
 	CurrentContext = nil;
 
 	// kick the whole buffer
-	[CurrentCommandBuffer addCompletedHandler:^(id <MTLCommandBuffer> Buffer)
+	[CurrentCommandBuffer addCompletedHandler : ^ (id <MTLCommandBuffer> Buffer)
 	 {
-		 dispatch_semaphore_signal(CommandBufferSemaphore);
+		dispatch_semaphore_signal(CommandBufferSemaphore);
 	 }];
-
+	
 	// Wait until at least 2 VBlanks has passed since last time
-	FramePacer->FramePacerEvent->Wait();
+//	FramePacer->FramePacerEvent->Wait();
 
 	// Commit before waiting to avoid leaving the gpu idle
 	[CurrentCommandBuffer commit];
@@ -427,10 +465,13 @@ void FMetalManager::SetBoundShaderState(FMetalBoundShaderState* BoundShaderState
 	CurrentBoundShaderState = BoundShaderState;
 }
 
-void FMetalManager::SetCurrentRenderTarget(FMetalSurface* RenderSurface, int32 RenderTargetIndex, int32 TotalNumRenderTargets)
+void FMetalManager::SetCurrentRenderTarget(FMetalSurface* RenderSurface, int32 RenderTargetIndex, uint32 MipIndex, uint32 ArraySliceIndex, MTLLoadAction LoadAction, MTLStoreAction StoreAction, int32 TotalNumRenderTargets)
 {
 	// rememeber our new max
 	CurrentNumRenderTargets = TotalNumRenderTargets;
+	
+	// user code generally passes -1 as a default, but we need 0
+	ArraySliceIndex = ArraySliceIndex == 0xFFFFFFFF ? 0 : ArraySliceIndex;
 
 	// update the current rendered-to pixel format
 	if (RenderSurface)
@@ -451,8 +492,17 @@ void FMetalManager::SetCurrentRenderTarget(FMetalSurface* RenderSurface, int32 R
 			// set the texture into the backbuffer
 			RenderSurface->Texture = CurrentDrawable.texture;
 		}
+		
+		if (RenderSurface->bIsCubemap)
+		{
+			ArraySliceIndex = GetMetalCubeFace((ECubeFace)ArraySliceIndex);
+		}
 
 		CurrentColorRenderTextures[RenderTargetIndex] = RenderSurface->Texture;
+		CurrentRenderTargetsViewInfo[RenderTargetIndex].MipIndex = MipIndex;
+		CurrentRenderTargetsViewInfo[RenderTargetIndex].ArraySliceIndex = ArraySliceIndex;
+		CurrentRenderTargetsViewInfo[RenderTargetIndex].LoadAction = LoadAction;
+		CurrentRenderTargetsViewInfo[RenderTargetIndex].StoreAction = StoreAction;
 
 		// only allow one MRT when using MSAA
 		checkf(RenderSurface->MSAATexture == NULL || TotalNumRenderTargets == 1);
@@ -461,6 +511,10 @@ void FMetalManager::SetCurrentRenderTarget(FMetalSurface* RenderSurface, int32 R
 	else
 	{
 		CurrentColorRenderTextures[RenderTargetIndex] = nil;
+		CurrentRenderTargetsViewInfo[RenderTargetIndex].MipIndex = 0;
+		CurrentRenderTargetsViewInfo[RenderTargetIndex].ArraySliceIndex = 0;
+		CurrentRenderTargetsViewInfo[RenderTargetIndex].LoadAction = MTLLoadActionDontCare;
+		CurrentRenderTargetsViewInfo[RenderTargetIndex].StoreAction = MTLStoreActionStore;
 	}
 }
 
@@ -486,7 +540,9 @@ void FMetalManager::UpdateContext()
 		bool bAllMatch = true;
 		for (uint32 AttachmentIndex = 0; AttachmentIndex < CurrentNumRenderTargets; AttachmentIndex++)
 		{
-			if (CurrentColorRenderTextures[AttachmentIndex] != PreviousColorRenderTextures[AttachmentIndex])
+			if (CurrentColorRenderTextures[AttachmentIndex] != PreviousColorRenderTextures[AttachmentIndex] ||
+				CurrentRenderTargetsViewInfo[AttachmentIndex].MipIndex != PreviousRenderTargetsViewInfo[AttachmentIndex].MipIndex ||
+				CurrentRenderTargetsViewInfo[AttachmentIndex].ArraySliceIndex != PreviousRenderTargetsViewInfo[AttachmentIndex].ArraySliceIndex)
 			{
 				bAllMatch = false;
 			}
@@ -527,7 +583,6 @@ void FMetalManager::UpdateContext()
 			{
 				// set up an MSAA attachment
 				ColorAttachment.texture = CurrentMSAARenderTexture;
-				[ColorAttachment setLoadAction:MTLLoadActionDontCare];
 				[ColorAttachment setStoreAction:MTLStoreActionMultisampleResolve];
 				[ColorAttachment setResolveTexture:CurrentColorRenderTextures[AttachmentIndex]];
 				Pipeline.SampleCount = CurrentMSAARenderTexture.sampleCount;
@@ -539,10 +594,13 @@ void FMetalManager::UpdateContext()
 			{
 				// set up non-MSAA attachment
 				ColorAttachment.texture = CurrentColorRenderTextures[AttachmentIndex];
-				[ColorAttachment setLoadAction:MTLLoadActionDontCare];
 				[ColorAttachment setStoreAction:MTLStoreActionStore];
 				Pipeline.SampleCount = 1;
 			}
+			ColorAttachment.level = CurrentRenderTargetsViewInfo[AttachmentIndex].MipIndex;
+			ColorAttachment.slice = CurrentRenderTargetsViewInfo[AttachmentIndex].ArraySliceIndex;
+			ColorAttachment.LoadAction = CurrentRenderTargetsViewInfo[AttachmentIndex].LoadAction;
+			//@todo implement store, but making sure that multisampleresolve is handled properly.
 
 			// assign the attachment to the slot
 			[CurrentRenderPass.colorAttachments setObject:ColorAttachment atIndexedSubscript:AttachmentIndex];
@@ -560,6 +618,7 @@ void FMetalManager::UpdateContext()
 
 		// remember this for next time
 		PreviousColorRenderTextures[AttachmentIndex] = CurrentColorRenderTextures[AttachmentIndex];
+		PreviousRenderTargetsViewInfo[AttachmentIndex] = CurrentRenderTargetsViewInfo[AttachmentIndex];
 	}
 	
 	if (CurrentDepthRenderTexture)
@@ -597,10 +656,12 @@ void FMetalManager::UpdateContext()
 	// commit pending commands on the old render target
 	if (CurrentContext)
 	{
+#if UE_BUILD_DEBUG || UE_BUILD_DEVELOPMENT
 		if (NumDrawCalls == 0)
 		{
 			NSLog(@"There were %d draw calls for an RT in frame %lld", NumDrawCalls, GFrameCounter);
 		}
+#endif
 
 		[CurrentContext endEncoding];
 		NumDrawCalls = 0;
@@ -613,18 +674,9 @@ void FMetalManager::UpdateContext()
 		[CurrentCommandBuffer commit];
 		UNTRACK_OBJECT(CurrentCommandBuffer);
 		[CurrentCommandBuffer release];
-
+		
 		// create the command buffer for this frame
-		CurrentCommandBuffer = [CommandQueue commandBufferWithUnretainedReferences];
-		[CurrentCommandBuffer retain];
-		TRACK_OBJECT(CurrentCommandBuffer);
-
-		uint64 LocalCommandBufferIndex = CommandBufferIndex++;
-
-		[CurrentCommandBuffer addCompletedHandler : ^ (id <MTLCommandBuffer> Buffer)
-		{
-			FMetalManager::Get()->SetCompletedCommandBufferIndex(LocalCommandBufferIndex);
-		}];
+		CreateCurrentCommandBuffer(false);
 	}
 		
 	// make a new render context to use to render to the framebuffer
@@ -671,7 +723,7 @@ uint32 FMetalManager::AllocateFromRingBuffer(uint32 Size, uint32 Alignment)
 
 uint32 FMetalManager::AllocateFromQueryBuffer()
 {
- 	return QueryBuffer.Allocate(8, 0);
+	return QueryBuffer.Allocate(8, 0);
 }
 
 
diff --git a/Engine/Source/Runtime/IOS/MetalRHI/Private/MetalManager.h b/Engine/Source/Runtime/IOS/MetalRHI/Private/MetalManager.h
--- a/Engine/Source/Runtime/IOS/MetalRHI/Private/MetalManager.h
+++ b/Engine/Source/Runtime/IOS/MetalRHI/Private/MetalManager.h
@@ -99,7 +99,7 @@ public:
 	 */
 	void SetBlendState(FMetalBlendState* BlendState);
 	void SetBoundShaderState(FMetalBoundShaderState* BoundShaderState);
-	void SetCurrentRenderTarget(FMetalSurface* RenderSurface, int32 RenderTargetIndex, int32 TotalNumRenderTargets);
+	void SetCurrentRenderTarget(FMetalSurface* RenderSurface, int32 RenderTargetIndex, uint32 MipIndex, uint32 ArraySliceIndex, MTLLoadAction LoadAction, MTLStoreAction StoreAction, int32 TotalNumRenderTargets);
 	void SetCurrentDepthStencilTarget(FMetalSurface* RenderSurface);
 	
 	/**
@@ -132,6 +132,7 @@ public:
 		CompletedCommandBufferIndex = Index;
 	}
 
+	void SubmitCommandBufferAndWait();
 	bool WaitForCommandBufferComplete(uint64 IndexToWaitFor, double Timeout);
 
 	void SetRasterizerState(const FRasterizerStateInitializerRHI& State);
@@ -162,6 +163,9 @@ protected:
 	void InitFrame();
 	void GenerateFetchShader();
 
+	void CreateCurrentCommandBuffer(bool bWait);
+
+
 	id<MTLDevice> Device;
 
 	id<MTLCommandQueue> CommandQueue;
@@ -174,6 +178,15 @@ protected:
 
 	id<MTLRenderCommandEncoder> CurrentContext;
 	
+	struct FRenderTargetViewInfo
+	{
+		uint32 MipIndex;
+		uint32 ArraySliceIndex;
+		MTLLoadAction LoadAction;
+		MTLStoreAction StoreAction;
+	};
+	
+	FRenderTargetViewInfo CurrentRenderTargetsViewInfo[MaxMetalRenderTargets], PreviousRenderTargetsViewInfo[MaxMetalRenderTargets];
 	uint32 CurrentNumRenderTargets, PreviousNumRenderTargets;
 	id<MTLTexture> CurrentColorRenderTextures[MaxMetalRenderTargets], PreviousColorRenderTextures[MaxMetalRenderTargets];
 	id<MTLTexture> CurrentDepthRenderTexture, PreviousDepthRenderTexture;
diff --git a/Engine/Source/Runtime/IOS/MetalRHI/Private/MetalRHIPrivate.h b/Engine/Source/Runtime/IOS/MetalRHI/Private/MetalRHIPrivate.h
--- a/Engine/Source/Runtime/IOS/MetalRHI/Private/MetalRHIPrivate.h
+++ b/Engine/Source/Runtime/IOS/MetalRHI/Private/MetalRHIPrivate.h
@@ -261,3 +261,38 @@ public:
 		return bRetVal;
 	}
 };
+
+FORCEINLINE int32 GetMetalCubeFace(ECubeFace Face)
+{
+	switch (Face)
+	{
+	case CubeFace_NegX:	return 0;
+	case CubeFace_NegY:	return 1;
+	case CubeFace_NegZ:	return 2;
+	case CubeFace_PosX:	return 3;
+	case CubeFace_PosY:	return 4;
+	case CubeFace_PosZ:	return 5;
+	default:			return -1;
+	}
+}
+
+FORCEINLINE MTLLoadAction GetMetalRTLoadAction(ERenderTargetLoadAction LoadAction)
+{
+	switch(LoadAction)
+	{
+		case ERenderTargetLoadAction::ENoAction: return MTLLoadActionDontCare;
+		case ERenderTargetLoadAction::ELoad: return MTLLoadActionLoad;
+		default: return MTLLoadActionDontCare;
+	}
+}
+
+FORCEINLINE MTLStoreAction GetMetalRTStoreAction(ERenderTargetStoreAction StoreAction)
+{
+	switch(StoreAction)
+	{
+		case ERenderTargetStoreAction::ENoAction: return MTLStoreActionDontCare;
+		case ERenderTargetStoreAction::EStore: return MTLStoreActionStore;
+		case ERenderTargetStoreAction::EMultisampleResolve: return MTLStoreActionMultisampleResolve;
+		default: return MTLStoreActionDontCare;
+	}
+}
\ No newline at end of file
diff --git a/Engine/Source/Runtime/IOS/MetalRHI/Private/MetalRenderTarget.cpp b/Engine/Source/Runtime/IOS/MetalRHI/Private/MetalRenderTarget.cpp
--- a/Engine/Source/Runtime/IOS/MetalRHI/Private/MetalRenderTarget.cpp
+++ b/Engine/Source/Runtime/IOS/MetalRHI/Private/MetalRenderTarget.cpp
@@ -1,3 +1,4 @@
+
 // Copyright 1998-2014 Epic Games, Inc. All Rights Reserved.
 
 /*=============================================================================
@@ -30,7 +31,39 @@ void FMetalDynamicRHI::RHIUnmapStagingSurface(FTextureRHIParamRef TextureRHI)
 
 void FMetalDynamicRHI::RHIReadSurfaceFloatData(FTextureRHIParamRef TextureRHI, FIntRect Rect, TArray<FFloat16Color>& OutData, ECubeFace CubeFace,int32 ArrayIndex,int32 MipIndex)
 {
+	//kick the current command buffer.
+	FMetalManager::Get()->SubmitCommandBufferAndWait();
+
+	DYNAMIC_CAST_METGALRESOURCE(Texture2D, Texture);
+	
+	FMetalSurface& Surface = GetMetalSurfaceFromRHITexture(TextureRHI);
+	
+	// verify the input image format (but don't crash)
+	if (Surface.PixelFormat != PF_FloatRGBA)
+	{
+		UE_LOG(LogRHI, Log, TEXT("Trying to read non-FloatRGBA surface."));
+	}
 
+	if (TextureRHI->GetTextureCube())
+	{
+		// adjust index to account for cubemaps as texture arrays
+		ArrayIndex *= CubeFace_MAX;
+		ArrayIndex += GetMetalCubeFace(CubeFace);
+	}
+	
+	// allocate output space
+	const uint32 SizeX = Rect.Width();
+	const uint32 SizeY = Rect.Height();
+	const uint32 Stride = 0; //can set 0 to indicate tightly packed.
+	const uint32 BytesPerImage = 0;
+	OutData.Empty();
+	OutData.AddUninitialized(SizeX * SizeY);
+	
+	FFloat16Color* OutDataPtr = OutData.GetData();
+	MTLRegion Region = MTLRegionMake2D(Rect.Min.X, Rect.Min.Y, Rect.Max.X, Rect.Max.Y);
+	
+	// function wants details about the destination, not the source
+	[Surface.Texture getBytes: OutDataPtr bytesPerRow:Stride bytesPerImage:BytesPerImage fromRegion:Region mipmapLevel:MipIndex slice:ArrayIndex];
 }
 
 void FMetalDynamicRHI::RHIRead3DSurfaceFloatData(FTextureRHIParamRef TextureRHI,FIntRect InRect,FIntPoint ZMinMax,TArray<FFloat16Color>& OutData)
diff --git a/Engine/Source/Runtime/IOS/MetalRHI/Private/MetalShaders.cpp b/Engine/Source/Runtime/IOS/MetalRHI/Private/MetalShaders.cpp
--- a/Engine/Source/Runtime/IOS/MetalRHI/Private/MetalShaders.cpp
+++ b/Engine/Source/Runtime/IOS/MetalRHI/Private/MetalShaders.cpp
@@ -427,7 +427,8 @@ void FMetalShaderParameterCache::CommitPackedGlobals(int32 Stage, const FMetalSh
 		// is there any data that needs to be copied?
 		if (PackedGlobalUniformDirty[UniformBufferIndex].HighVector > 0)//PackedGlobalUniformDirty[UniformBufferIndex].LowVector)
 		{
-			check(UniformBufferIndex == 0 || UniformBufferIndex == 1);
+			//@todo rco safe to remove this forever?
+			//check(UniformBufferIndex == 0 || UniformBufferIndex == 1 || UniformBufferIndex == 3);
 			uint32 TotalSize = Bindings.PackedGlobalArrays[Index].Size;
 			uint32 SizeToUpload = PackedGlobalUniformDirty[UniformBufferIndex].HighVector * SizeOfFloat4;
 //@todo-rco: Temp workaround
@@ -529,23 +530,4 @@ void FMetalShaderParameterCache::CommitPackedUniformBuffers(TRefCountPtr<FMetalB
 			}
 		}
 	}
-	else
-	{
-		for (int32 BufferIndex = 0; BufferIndex < RHIUniformBuffers.Num(); ++BufferIndex)
-		{
-			const FRHIUniformBuffer* RHIUniformBuffer = RHIUniformBuffers[BufferIndex];
-			if (RHIUniformBuffer)
-			{
-				auto* UB = (FMetalUniformBuffer*)RHIUniformBuffer;
-				if (Stage == CrossCompiler::SHADER_STAGE_VERTEX)
-				{
-					[FMetalManager::GetContext() setVertexBuffer:UB->Buffer offset:UB->Offset atIndex:BufferIndex];
-				}
-				else
-				{
-					[FMetalManager::GetContext() setFragmentBuffer:UB->Buffer offset:UB->Offset atIndex:BufferIndex];
-				}
-			}
-		}
-	}
 }
diff --git a/Engine/Source/Runtime/IOS/MetalRHI/Private/MetalTexture.cpp b/Engine/Source/Runtime/IOS/MetalRHI/Private/MetalTexture.cpp
--- a/Engine/Source/Runtime/IOS/MetalRHI/Private/MetalTexture.cpp
+++ b/Engine/Source/Runtime/IOS/MetalRHI/Private/MetalTexture.cpp
@@ -6,20 +6,6 @@
 
 #include "MetalRHIPrivate.h"
 
-FORCEINLINE int32 GetMetalCubeFace(ECubeFace Face)
-{
-	switch (Face)
-	{
-		case CubeFace_NegX:	return 0;
-		case CubeFace_NegY:	return 1;
-		case CubeFace_NegZ:	return 2;
-		case CubeFace_PosX:	return 3;
-		case CubeFace_PosY:	return 4;
-		case CubeFace_PosZ:	return 5;
-		default:			return -1;
-	}
-}
-
 /** Texture reference class. */
 class FMetalTextureReference : public FRHITextureReference
 {
@@ -448,7 +434,7 @@ FTexture2DRHIRef FMetalDynamicRHI::RHIAsyncReallocateTexture2D(FTexture2DRHIPara
 					 sourceSlice:SliceIndex
 					 sourceLevel:MipIndex + SourceMipOffset
 					sourceOrigin:Origin
-					  sourceSize:MTLSizeMake(MipSizeX, MipSizeY, 0)
+					  sourceSize:MTLSizeMake(MipSizeX, MipSizeY, 1)
 					   toTexture:NewTexture->Surface.Texture
 			    destinationSlice:SliceIndex
 			    destinationLevel:MipIndex + DestMipOffset
