Fix for Metal RHI compile error

diff --git a/Engine/Source/Runtime/IOS/MetalRHI/Private/MetalRHI.cpp b/Engine/Source/Runtime/IOS/MetalRHI/Private/MetalRHI.cpp
--- a/Engine/Source/Runtime/IOS/MetalRHI/Private/MetalRHI.cpp
+++ b/Engine/Source/Runtime/IOS/MetalRHI/Private/MetalRHI.cpp
@@ -42,9 +42,8 @@ FMetalDynamicRHI::FMetalDynamicRHI()
 	check( !GIsThreadedRendering );
 
 	// Initialize the RHI capabilities.
-	GRHIFeatureLevelValue = ERHIFeatureLevel::ES2;
-	GRHIShaderPlatformValue = SP_METAL;
-	GMaxRHIFeatureLevelValue = GRHIFeatureLevelValue = ERHIFeatureLevel::ES2;
+	GMaxRHIFeatureLevelValue = ERHIFeatureLevel::ES2;
+	GMaxRHIShaderPlatformValue = SP_METAL;
 
 	GPixelCenterOffset = 0.0f;
 	GSupportsVertexInstancing = false;
