UE4 - Remove Metal horrible hack, use new SetRenderTargets api

diff --git a/Engine/Source/Runtime/IOS/MetalRHI/Private/MetalManager.cpp b/Engine/Source/Runtime/IOS/MetalRHI/Private/MetalManager.cpp
--- a/Engine/Source/Runtime/IOS/MetalRHI/Private/MetalManager.cpp
+++ b/Engine/Source/Runtime/IOS/MetalRHI/Private/MetalManager.cpp
@@ -36,10 +36,6 @@ static uint32 RTBitOffsets[] = { OFFSET_RENDER_TARGET_FORMAT0, OFFSET_RENDER_TAR
 	}
 #define GET_HASH(Offset, NumBits) ((Hash >> Offset) & ((1ULL << NumBits) - 1))
 
-//@todo-rco: HACK! This needs to exist until we merge SetRenderTargets()+Clear()
-extern int GMetalHackDepthActions;
-
-
 /*
 void FPipelineShadow::SetHash(uint64 InHash)
 {
@@ -211,6 +207,8 @@ FMetalManager::FMetalManager()
 	}
 	FMemory::MemSet(CurrentRenderTargetsViewInfo, 0);
 	FMemory::MemSet(PreviousRenderTargetsViewInfo, 0);
+	FMemory::MemSet(CurrentDepthViewInfo, 0);
+	FMemory::MemSet(PreviousDepthViewInfo, 0);
 
 	CommandQueue = [Device newCommandQueue];
 
@@ -519,16 +517,22 @@ void FMetalManager::SetCurrentRenderTarget(FMetalSurface* RenderSurface, int32 R
 	}
 }
 
-void FMetalManager::SetCurrentDepthStencilTarget(FMetalSurface* RenderSurface)
+void FMetalManager::SetCurrentDepthStencilTarget(FMetalSurface* RenderSurface, MTLLoadAction LoadAction, MTLStoreAction StoreAction, float ClearDepthValue)
 {
 	if (RenderSurface)
 	{
 		// @todo metal stencil: track stencil here
 		CurrentDepthRenderTexture = RenderSurface->Texture;
+		CurrentDepthViewInfo.LoadAction = LoadAction;
+		CurrentDepthViewInfo.StoreAction = StoreAction;
+		CurrentDepthViewInfo.ClearDepthValue = ClearDepthValue;
 	}
 	else
 	{
 		CurrentDepthRenderTexture = nil;
+		CurrentDepthViewInfo.LoadAction = MTLLoadActionClear;
+		CurrentDepthViewInfo.StoreAction = MTLStoreActionDontCare;
+		CurrentDepthViewInfo.ClearDepthValue = 0.0f;
 	}
 }
 
@@ -553,6 +557,8 @@ void FMetalManager::UpdateContext()
 		{
 			return;
 		}
+
+		//@todo-rco: Do we need to test changes in Load/Store actions (and/or clear values) for Color & Depth?
 	}
 
 	PreviousNumRenderTargets = CurrentNumRenderTargets;
@@ -628,19 +634,9 @@ void FMetalManager::UpdateContext()
 
 		// set up the depth attachment
 		DepthAttachment.texture = CurrentDepthRenderTexture;
-		[DepthAttachment setLoadAction:MTLLoadActionClear];
-		if (GMetalHackDepthActions == 1)
-		{
-			// For Shadow Maps, clear to 1.0f and save the render target
-			[DepthAttachment setStoreAction:MTLStoreActionStore];
-			[DepthAttachment setClearDepth:1.0f];
-		}
-		else if (GMetalHackDepthActions == 0)
-		{
-			// Regular operation, clear to zero (due to inverted Z) and don't store it
-			[DepthAttachment setStoreAction:MTLStoreActionDontCare];
-			[DepthAttachment setClearDepth:0.0f];
-		}
+		[DepthAttachment setLoadAction:CurrentDepthViewInfo.LoadAction];
+		[DepthAttachment setStoreAction : CurrentDepthViewInfo.StoreAction];
+		[DepthAttachment setClearDepth : CurrentDepthViewInfo.ClearDepthValue];
 
 		Pipeline.DepthTargetFormat = CurrentDepthRenderTexture.pixelFormat;
 		if (Pipeline.SampleCount == 0)
@@ -660,6 +656,7 @@ void FMetalManager::UpdateContext()
 	SET_HASH(OFFSET_DEPTH_TARGET_FORMAT, NUMBITS_DEPTH_TARGET_FORMAT, Pipeline.DepthTargetFormat);
 
 	// remember this for next time
+	PreviousDepthViewInfo = CurrentDepthViewInfo;
 	PreviousDepthRenderTexture = CurrentDepthRenderTexture;
 
 	SET_HASH(OFFSET_SAMPLE_COUNT, NUMBITS_SAMPLE_COUNT, Pipeline.SampleCount);
diff --git a/Engine/Source/Runtime/IOS/MetalRHI/Private/MetalManager.h b/Engine/Source/Runtime/IOS/MetalRHI/Private/MetalManager.h
--- a/Engine/Source/Runtime/IOS/MetalRHI/Private/MetalManager.h
+++ b/Engine/Source/Runtime/IOS/MetalRHI/Private/MetalManager.h
@@ -100,7 +100,7 @@ public:
 	void SetBlendState(FMetalBlendState* BlendState);
 	void SetBoundShaderState(FMetalBoundShaderState* BoundShaderState);
 	void SetCurrentRenderTarget(FMetalSurface* RenderSurface, int32 RenderTargetIndex, uint32 MipIndex, uint32 ArraySliceIndex, MTLLoadAction LoadAction, MTLStoreAction StoreAction, int32 TotalNumRenderTargets);
-	void SetCurrentDepthStencilTarget(FMetalSurface* RenderSurface);
+	void SetCurrentDepthStencilTarget(FMetalSurface* RenderSurface, MTLLoadAction LoadAction, MTLStoreAction StoreAction, float ClearDepthValue);
 	
 	/**
 	 * Update the context with the current render targets
@@ -185,10 +185,18 @@ protected:
 		MTLLoadAction LoadAction;
 		MTLStoreAction StoreAction;
 	};
-	
+
+	struct FDepthRenderTargetViewInfo
+	{
+		float ClearDepthValue;
+		MTLLoadAction LoadAction;
+		MTLStoreAction StoreAction;
+	};
+
 	FRenderTargetViewInfo CurrentRenderTargetsViewInfo[MaxMetalRenderTargets], PreviousRenderTargetsViewInfo[MaxMetalRenderTargets];
 	uint32 CurrentNumRenderTargets, PreviousNumRenderTargets;
 	id<MTLTexture> CurrentColorRenderTextures[MaxMetalRenderTargets], PreviousColorRenderTextures[MaxMetalRenderTargets];
+	FDepthRenderTargetViewInfo CurrentDepthViewInfo, PreviousDepthViewInfo;
 	id<MTLTexture> CurrentDepthRenderTexture, PreviousDepthRenderTexture;
 	id<MTLTexture> CurrentMSAARenderTexture;
 	
diff --git a/Engine/Source/Runtime/IOS/MetalRHI/Private/MetalRHIPrivate.h b/Engine/Source/Runtime/IOS/MetalRHI/Private/MetalRHIPrivate.h
--- a/Engine/Source/Runtime/IOS/MetalRHI/Private/MetalRHIPrivate.h
+++ b/Engine/Source/Runtime/IOS/MetalRHI/Private/MetalRHIPrivate.h
@@ -282,6 +282,7 @@ FORCEINLINE MTLLoadAction GetMetalRTLoadAction(ERenderTargetLoadAction LoadActio
 	{
 		case ERenderTargetLoadAction::ENoAction: return MTLLoadActionDontCare;
 		case ERenderTargetLoadAction::ELoad: return MTLLoadActionLoad;
+		case ERenderTargetLoadAction::EClear: return MTLLoadActionClear;
 		default: return MTLLoadActionDontCare;
 	}
 }
diff --git a/Engine/Source/Runtime/OpenGLDrv/Private/OpenGLCommands.cpp b/Engine/Source/Runtime/OpenGLDrv/Private/OpenGLCommands.cpp
--- a/Engine/Source/Runtime/OpenGLDrv/Private/OpenGLCommands.cpp
+++ b/Engine/Source/Runtime/OpenGLDrv/Private/OpenGLCommands.cpp
@@ -1637,7 +1637,7 @@ void FOpenGLDynamicRHI::RHISetRenderTargetsAndClear(const FRHISetRenderTargetsIn
 {
 	this->RHISetRenderTargets(RenderTargetsInfo.NumColorRenderTargets,
 		RenderTargetsInfo.ColorRenderTarget,
-		RenderTargetsInfo.DepthStencilTarget,
+		RenderTargetsInfo.DepthStencilRenderTarget.Texture,
 		0,
 		nullptr);
 	if (RenderTargetsInfo.bClearColor || RenderTargetsInfo.bClearStencil || RenderTargetsInfo.bClearDepth)
diff --git a/Engine/Source/Runtime/RHI/Public/RHIResources.h b/Engine/Source/Runtime/RHI/Public/RHIResources.h
--- a/Engine/Source/Runtime/RHI/Public/RHIResources.h
+++ b/Engine/Source/Runtime/RHI/Public/RHIResources.h
@@ -621,35 +621,85 @@ public:
 	{}
 };
 
+class FRHIDepthRenderTargetView
+{
+public:
+	FTextureRHIParamRef Texture;
+
+	ERenderTargetLoadAction LoadAction;
+	ERenderTargetStoreAction StoreAction;
+
+	FRHIDepthRenderTargetView() :
+		Texture(nullptr),
+		LoadAction(ERenderTargetLoadAction::EClear),
+		StoreAction(ERenderTargetStoreAction::ENoAction)
+	{}
+
+	FRHIDepthRenderTargetView(const FRHIDepthRenderTargetView& Other) :
+		Texture(Other.Texture),
+		LoadAction(Other.LoadAction),
+		StoreAction(Other.StoreAction)
+	{}
+
+	FRHIDepthRenderTargetView(FTextureRHIParamRef InTexture) :
+		Texture(InTexture),
+		LoadAction(ERenderTargetLoadAction::EClear),
+		StoreAction(ERenderTargetStoreAction::ENoAction)
+	{}
+
+	FRHIDepthRenderTargetView(FTextureRHIParamRef InTexture, ERenderTargetLoadAction InLoadAction, ERenderTargetStoreAction InStoreAction) :
+		Texture(InTexture),
+		LoadAction(InLoadAction),
+		StoreAction(InStoreAction)
+	{}
+};
+
 class FRHISetRenderTargetsInfo
 {
 public:
 	// Color Render Targets Info
-	int32 NumColorRenderTargets;
+	FRHIRenderTargetView ColorRenderTarget[MaxSimultaneousRenderTargets];
 	FLinearColor ClearColors[MaxSimultaneousRenderTargets];
+	int32 NumColorRenderTargets;
 	bool bClearColor;
-	FRHIRenderTargetView ColorRenderTarget[MaxSimultaneousRenderTargets];
 
 	// Depth/Stencil Render Target Info
-	ERenderTargetLoadAction DepthLoadAction;
-	ERenderTargetStoreAction DepthStoreAction;
+	FRHIDepthRenderTargetView DepthStencilRenderTarget;
 	float DepthClearValue;
-	bool bClearDepth;
 	uint32 StencilClearValue;
+	bool bClearDepth;
 	bool bClearStencil;
-	FTextureRHIParamRef DepthStencilTarget;
 
 	FRHISetRenderTargetsInfo() :
 		NumColorRenderTargets(0),
 		bClearColor(false),
-		DepthLoadAction(ERenderTargetLoadAction::EClear),
-		DepthStoreAction(ERenderTargetStoreAction::ENoAction),
 		DepthClearValue(0.0f),
-		bClearDepth(false),
 		StencilClearValue(0),
-		bClearStencil(false),
-		DepthStencilTarget(nullptr)
+		bClearDepth(false),
+		bClearStencil(false)
+	{}
+
+	FRHISetRenderTargetsInfo(int32 InNumColorRenderTargets, const FRHIRenderTargetView* InColorRenderTargets, const FRHIDepthRenderTargetView& InDepthStencilRenderTarget) :
+		NumColorRenderTargets(InNumColorRenderTargets),
+		bClearColor(false),
+		DepthStencilRenderTarget(InDepthStencilRenderTarget),
+		bClearDepth(false),
+		bClearStencil(false)
+	{
+		check(InNumColorRenderTargets <= 0 || InColorRenderTargets);
+		for (int32 Index = 0; Index < InNumColorRenderTargets; ++Index)
+		{
+			ColorRenderTarget[Index] = InColorRenderTargets[Index];
+		}
+	}
+
+	void SetClearDepthStencil(bool bInClearDepth, float InDepthClear, bool bInClearStencil = false, uint32 InClearStencil = 0)
 	{
+		DepthStencilRenderTarget.LoadAction = ERenderTargetLoadAction::EClear;
+		bClearDepth = bInClearDepth;
+		DepthClearValue = InDepthClear;
+		bClearStencil = bInClearStencil;
+		StencilClearValue = InClearStencil;
 	}
 };
 
diff --git a/Engine/Source/Runtime/Renderer/Private/PostProcess/SceneRenderTargets.cpp b/Engine/Source/Runtime/Renderer/Private/PostProcess/SceneRenderTargets.cpp
--- a/Engine/Source/Runtime/Renderer/Private/PostProcess/SceneRenderTargets.cpp
+++ b/Engine/Source/Runtime/Renderer/Private/PostProcess/SceneRenderTargets.cpp
@@ -740,14 +740,17 @@ void FSceneRenderTargets::FinishRenderingPrePass(FRHICommandListImmediate& RHICm
 void FSceneRenderTargets::BeginRenderingShadowDepth(FRHICommandListImmediate& RHICmdList)
 {
 	GRenderTargetPool.VisualizeTexture.SetCheckPoint(RHICmdList, ShadowDepthZ);
-	if (GSupportsDepthRenderTargetWithoutColorRenderTarget)
-	{
-		SetRenderTarget(RHICmdList, FTextureRHIRef(), GetShadowDepthZSurface());
-	}
-	else
+	FRHISetRenderTargetsInfo Info(0, nullptr, FRHIDepthRenderTargetView(GetShadowDepthZSurface(), ERenderTargetLoadAction::EClear, ERenderTargetStoreAction::EStore));
+	Info.SetClearDepthStencil(true, 1.0f);
+	Info.ColorRenderTarget[0].StoreAction = ERenderTargetStoreAction::ENoAction;
+
+	if (!GSupportsDepthRenderTargetWithoutColorRenderTarget)
 	{
-		SetRenderTarget(RHICmdList, GetOptionalShadowDepthColorSurface(), GetShadowDepthZSurface());
+		Info.NumColorRenderTargets = 1;
+		Info.ColorRenderTarget[0].Texture = GetOptionalShadowDepthColorSurface();
 	}
+
+	RHICmdList.SetRenderTargetsAndClear(Info);
 }
 
 void FSceneRenderTargets::BeginRenderingCubeShadowDepth(FRHICommandListImmediate& RHICmdList, int32 ShadowResolution)
diff --git a/Engine/Source/Runtime/Renderer/Private/ShadowRendering.cpp b/Engine/Source/Runtime/Renderer/Private/ShadowRendering.cpp
--- a/Engine/Source/Runtime/Renderer/Private/ShadowRendering.cpp
+++ b/Engine/Source/Runtime/Renderer/Private/ShadowRendering.cpp
@@ -3317,12 +3317,6 @@ bool FDeferredShadingSceneRenderer::RenderCachedPreshadows(FRHICommandListImmedi
 	return bAttenuationBufferDirty;
 }
 
-#if PLATFORM_IOS
-// 0: Regular operation, clear to zero (due to inverted Z) and don't store it
-// 1: For Shadow Maps, clear to 1.0f and save the render target
-int GMetalHackDepthActions = 0;
-#endif
-
 /**
  * Used by RenderLights to render shadows to the attenuation buffer.
  *
@@ -3395,17 +3389,9 @@ bool FForwardShadingSceneRenderer::RenderShadowDepthMap(FRHICommandListImmediate
 	{
 		// Render the shadow depths.
 		SCOPED_DRAW_EVENT(RHICmdList, ShadowDepthsFromOpaque, DEC_SCENE_ITEMS);
-#if PLATFORM_IOS
-		// WILL BE FIXED WHEN WE HAVE SET & CLEAR RT IN THE SAME API! Also a way to tell it to be resolved!
-		GMetalHackDepthActions = 1;
-#endif
-		GSceneRenderTargets.BeginRenderingShadowDepth(RHICmdList);
-#if PLATFORM_IOS
-		// WILL BE FIXED WHEN WE HAVE SET & CLEAR RT IN THE SAME API!
-		GMetalHackDepthActions = 0;
-#endif
 
-		RHICmdList.Clear(false, FColor(255, 255, 255), true, 1.0f, false, 0, FIntRect());
+		// Clearin the render target happens inside the Begin function
+		GSceneRenderTargets.BeginRenderingShadowDepth(RHICmdList);
 
 		// render depth for each shadow
 		for (int32 ShadowIndex = 0; ShadowIndex < Shadows.Num(); ShadowIndex++)
diff --git a/Engine/Source/Runtime/Windows/D3D11RHI/Private/D3D11Commands.cpp b/Engine/Source/Runtime/Windows/D3D11RHI/Private/D3D11Commands.cpp
--- a/Engine/Source/Runtime/Windows/D3D11RHI/Private/D3D11Commands.cpp
+++ b/Engine/Source/Runtime/Windows/D3D11RHI/Private/D3D11Commands.cpp
@@ -1007,7 +1007,7 @@ void FD3D11DynamicRHI::RHISetRenderTargetsAndClear(const FRHISetRenderTargetsInf
 {
 	this->RHISetRenderTargets(RenderTargetsInfo.NumColorRenderTargets,
 		RenderTargetsInfo.ColorRenderTarget,
-		RenderTargetsInfo.DepthStencilTarget,
+		RenderTargetsInfo.DepthStencilRenderTarget.Texture,
 		0,
 		nullptr);
 	if (RenderTargetsInfo.bClearColor || RenderTargetsInfo.bClearStencil || RenderTargetsInfo.bClearDepth)
