UE4 - Missed file (Metal SetRT & Clear)

diff --git a/Engine/Source/Runtime/IOS/MetalRHI/Private/MetalCommands.cpp b/Engine/Source/Runtime/IOS/MetalRHI/Private/MetalCommands.cpp
--- a/Engine/Source/Runtime/IOS/MetalRHI/Private/MetalCommands.cpp
+++ b/Engine/Source/Runtime/IOS/MetalRHI/Private/MetalCommands.cpp
@@ -6,6 +6,8 @@
 
 #include "MetalRHIPrivate.h"
 
+static const bool GUsesInvertedZ = true;
+
 static MTLPrimitiveType TranslatePrimitiveType(uint32 PrimitiveType)
 {
 	switch (PrimitiveType)
@@ -317,8 +319,6 @@ void FMetalDynamicRHI::RHISetRenderTargets(uint32 NumSimultaneousRenderTargets,
 	FTextureRHIParamRef NewDepthStencilTargetRHI, uint32 NumUAVs, const FUnorderedAccessViewRHIParamRef* UAVs)
 {
 	FMetalManager* Manager = FMetalManager::Get();
-	
-
 	for (int32 RenderTargetIndex = 0; RenderTargetIndex < MaxMetalRenderTargets; RenderTargetIndex++)
 	{
 		const FRHIRenderTargetView& RenderTargetView = NewRenderTargets[RenderTargetIndex];
@@ -337,11 +337,11 @@ void FMetalDynamicRHI::RHISetRenderTargets(uint32 NumSimultaneousRenderTargets,
 	if (NewDepthStencilTargetRHI)
 	{
 		FMetalSurface& Surface = GetMetalSurfaceFromRHITexture(NewDepthStencilTargetRHI);
-		Manager->SetCurrentDepthStencilTarget(&Surface);
+		Manager->SetCurrentDepthStencilTarget(&Surface, MTLLoadActionClear, MTLStoreActionDontCare, GUsesInvertedZ ? 0.0f : 1.0f);
 	}
 	else
 	{
-		Manager->SetCurrentDepthStencilTarget(NULL);
+		Manager->SetCurrentDepthStencilTarget(NULL, MTLLoadActionClear, MTLStoreActionDontCare, GUsesInvertedZ ? 0.0f : 1.0f);
 	}
 	
 	// now that we have a new render target, we need a new context to render to it!
@@ -366,15 +366,45 @@ void FMetalDynamicRHI::RHIDiscardRenderTargets(bool Depth, bool Stencil, uint32
 
 void FMetalDynamicRHI::RHISetRenderTargetsAndClear(const FRHISetRenderTargetsInfo& RenderTargetsInfo)
 {
-//@todo-rco: TEMP
-	this->RHISetRenderTargets(RenderTargetsInfo.NumColorRenderTargets,
-		RenderTargetsInfo.ColorRenderTarget,
-		RenderTargetsInfo.DepthStencilTarget,
-		0,
-		nullptr);
-	if (RenderTargetsInfo.bClearColor || RenderTargetsInfo.bClearStencil || RenderTargetsInfo.bClearDepth)
+	FMetalManager* Manager = FMetalManager::Get();
+	for (int32 RenderTargetIndex = 0; RenderTargetIndex < MaxMetalRenderTargets; RenderTargetIndex++)
+	{
+		const FRHIRenderTargetView& RenderTargetView = RenderTargetsInfo.ColorRenderTarget[RenderTargetIndex];
+		// update the current RTs
+		if (RenderTargetIndex < RenderTargetsInfo.NumColorRenderTargets && RenderTargetView.Texture != NULL)
+		{
+			FMetalSurface& Surface = GetMetalSurfaceFromRHITexture(RenderTargetView.Texture);
+			Manager->SetCurrentRenderTarget(&Surface, RenderTargetIndex, RenderTargetView.MipIndex, RenderTargetView.ArraySliceIndex, GetMetalRTLoadAction(RenderTargetView.LoadAction), GetMetalRTStoreAction(RenderTargetView.StoreAction), RenderTargetsInfo.NumColorRenderTargets);
+		}
+		else
+		{
+			Manager->SetCurrentRenderTarget(NULL, RenderTargetIndex, 0, 0, MTLLoadActionDontCare, MTLStoreActionStore, RenderTargetsInfo.NumColorRenderTargets);
+		}
+	}
+
+	if (RenderTargetsInfo.DepthStencilRenderTarget.Texture)
 	{
-		this->RHIClearMRT(RenderTargetsInfo.bClearColor, RenderTargetsInfo.NumColorRenderTargets, RenderTargetsInfo.ClearColors, RenderTargetsInfo.bClearDepth, RenderTargetsInfo.DepthClearValue, RenderTargetsInfo.bClearStencil, RenderTargetsInfo.StencilClearValue, FIntRect());
+		FMetalSurface& Surface = GetMetalSurfaceFromRHITexture(RenderTargetsInfo.DepthStencilRenderTarget.Texture);
+		Manager->SetCurrentDepthStencilTarget(&Surface, GetMetalRTLoadAction(RenderTargetsInfo.DepthStencilRenderTarget.LoadAction), GetMetalRTStoreAction(RenderTargetsInfo.DepthStencilRenderTarget.StoreAction), RenderTargetsInfo.DepthClearValue);
+	}
+	else
+	{
+		Manager->SetCurrentDepthStencilTarget(NULL, MTLLoadActionClear, MTLStoreActionDontCare, 1.0f);
+	}
+
+	// now that we have a new render target, we need a new context to render to it!
+	Manager->UpdateContext();
+
+	// Set the viewport to the full size of render target 0.
+	if (RenderTargetsInfo.ColorRenderTarget[0].Texture)
+	{
+		const FRHIRenderTargetView& RenderTargetView = RenderTargetsInfo.ColorRenderTarget[0];
+		FMetalSurface& RenderTarget = GetMetalSurfaceFromRHITexture(RenderTargetView.Texture);
+
+		uint32 Width = FMath::Max(RenderTarget.Texture.width >> RenderTargetView.MipIndex, (uint32)1);
+		uint32 Height = FMath::Max(RenderTarget.Texture.height >> RenderTargetView.MipIndex, (uint32)1);
+
+		RHISetViewport(0, 0, 0.0f, Width, Height, 1.0f);
 	}
 }
 
