First pass of the Platform RHI Frame Pacer.
Should evolve as we add more platforms
IOS implementaion added and hooked into CADisplayLink
Tappy has this enabled.

diff --git a/Engine/Source/Runtime/Core/Private/IOS/IOSPlatformFramePacer.cpp b/Engine/Source/Runtime/Core/Private/IOS/IOSPlatformFramePacer.cpp
new file mode 100644
--- /dev/null
+++ b/Engine/Source/Runtime/Core/Private/IOS/IOSPlatformFramePacer.cpp
@@ -0,0 +1,88 @@
+// Copyright 1998-2014 Epic Games, Inc. All Rights Reserved.
+
+#include "Core.h"
+
+
+// Collection of events listening for this trigger.
+static TArray<FEvent*> ListeningEvents;
+
+
+/*******************************************************************
+ * FIOSFramePacer implementation
+ *******************************************************************/
+
+@interface FIOSFramePacer : NSObject
+{
+    @public
+	FEvent *FramePacerEvent;
+}
+
+-(void)run:(id)param;
+-(void)signal:(id)param;
+
+@end
+
+
+@implementation FIOSFramePacer
+
+-(void)run:(id)param
+{
+	NSRunLoop *runloop = [NSRunLoop currentRunLoop];
+	CADisplayLink *displayLink = [CADisplayLink displayLinkWithTarget:self selector:@selector(signal:)];
+	displayLink.frameInterval = FIOSPlatformRHIFramePacer::FrameInterval;
+    
+	[displayLink addToRunLoop:runloop forMode:NSDefaultRunLoopMode];
+	[runloop run];
+}
+
+
+-(void)signal:(id)param
+{
+    for( auto& NextEvent : ListeningEvents )
+    {
+        NextEvent->Trigger();
+    }
+}
+
+@end
+
+
+
+/*******************************************************************
+ * FIOSPlatformRHIFramePacer implementation
+ *******************************************************************/
+
+uint32 FIOSPlatformRHIFramePacer::FrameInterval = 1;
+FIOSFramePacer* FIOSPlatformRHIFramePacer::FramePacer = nil;
+
+bool FIOSPlatformRHIFramePacer::IsEnabled()
+{
+    bool bIsRHIFramePacerEnabled = false;
+	GConfig->GetBool( TEXT( "PlatformFramePacer" ), TEXT( "IsEnabled" ), bIsRHIFramePacerEnabled, GEngineIni );
+    
+    
+    return bIsRHIFramePacerEnabled;
+}
+
+
+void FIOSPlatformRHIFramePacer::InitWithEvent(FEvent* TriggeredEvent, uint32 InFrameInterval)
+{
+    FrameInterval = InFrameInterval;
+    
+    // Create display link thread
+	FramePacer = [[FIOSFramePacer alloc] init];
+	[NSThread detachNewThreadSelector:@selector(run:) toTarget:FramePacer withObject:nil];
+    
+    // Only one supported for now, we may want more eventually.
+    ListeningEvents.Add( TriggeredEvent );
+}
+
+
+void FIOSPlatformRHIFramePacer::Destroy()
+{
+    if( FramePacer != nil )
+    {
+        [FramePacer release];
+        FramePacer = nil;
+    }
+}
\ No newline at end of file
diff --git a/Engine/Source/Runtime/Core/Public/GenericPlatform/GenericPlatformFramePacer.h b/Engine/Source/Runtime/Core/Public/GenericPlatform/GenericPlatformFramePacer.h
new file mode 100644
--- /dev/null
+++ b/Engine/Source/Runtime/Core/Public/GenericPlatform/GenericPlatformFramePacer.h
@@ -0,0 +1,34 @@
+// Copyright 1998-2014 Epic Games, Inc. All Rights Reserved.
+
+
+/*=============================================================================================
+	GenericPlatformFramePacer.h: Generic platform frame pacer classes
+==============================================================================================*/
+
+
+#pragma once
+
+
+/**
+ * Generic implementation for most platforms, these tend to be unused and unimplemented
+ **/
+struct CORE_API FGenericPlatformRHIFramePacer
+{
+    /**
+     * Should the Frame Pacer be enabled?
+     */
+    bool IsEnabled() { return false; }
+    
+    /**
+     * Init the RHI frame pacer
+     *
+     * @param InTriggeredEvent - The event we wish to trigger when the frame interval has been triggered by the hardware.
+     * @param InFrameInterval - How often should the event be triggered, in Frames.
+     */
+	static void InitWithEvent(class FEvent* InTriggeredEvent, uint32 InFrameInterval) {}
+    
+    /**
+     * Teardown the Frame Pacer.
+     */
+    static void Destroy() {}
+};
\ No newline at end of file
diff --git a/Engine/Source/Runtime/Core/Public/HAL/PlatformIncludes.h b/Engine/Source/Runtime/Core/Public/HAL/PlatformIncludes.h
--- a/Engine/Source/Runtime/Core/Public/HAL/PlatformIncludes.h
+++ b/Engine/Source/Runtime/Core/Public/HAL/PlatformIncludes.h
@@ -29,6 +29,7 @@
 #include "GenericPlatform/GenericPlatformSurvey.h"
 #include "GenericPlatform/GenericPlatformHttp.h"
 #include "GenericPlatform/GenericPlatformAffinity.h"
+#include "GenericPlatform/GenericPlatformFramePacer.h"
 
 // Platform specific include.
 //@port This is where you include your "main" platform header
diff --git a/Engine/Source/Runtime/Core/Public/HTML5/HTML5PlatformIncludes.h b/Engine/Source/Runtime/Core/Public/HTML5/HTML5PlatformIncludes.h
--- a/Engine/Source/Runtime/Core/Public/HTML5/HTML5PlatformIncludes.h
+++ b/Engine/Source/Runtime/Core/Public/HTML5/HTML5PlatformIncludes.h
@@ -58,6 +58,8 @@ struct RECT
 #include "HTML5/HTML5PlatformSurvey.h"
 #include "HTML5/HTML5PlatformHttp.h"
 
+typedef FGenericPlatformRHIFramePacer FPlatformRHIFramePacer;
+
 typedef FGenericPlatformAffinity FPlatformAffinity;
 
 // include platform properties and typedef it for the runtime
diff --git a/Engine/Source/Runtime/Core/Public/IOS/IOSPlatformFramePacer.h b/Engine/Source/Runtime/Core/Public/IOS/IOSPlatformFramePacer.h
new file mode 100644
--- /dev/null
+++ b/Engine/Source/Runtime/Core/Public/IOS/IOSPlatformFramePacer.h
@@ -0,0 +1,33 @@
+// Copyright 1998-2014 Epic Games, Inc. All Rights Reserved.
+
+
+/*=============================================================================================
+	IOSPlatformFramePacer.h: Apple iOS platform frame pacer classes.
+==============================================================================================*/
+
+
+#pragma once
+
+// Forward declare the ios frame pacer class we will be using.
+@class FIOSFramePacer;
+
+
+/**
+ * iOS implementation of FGenericPlatformRHIFramePacer
+ **/
+struct FIOSPlatformRHIFramePacer : public FGenericPlatformRHIFramePacer
+{
+    // FGenericPlatformRHIFramePacer interface
+    static bool IsEnabled();
+	static void InitWithEvent(class FEvent* TriggeredEvent, uint32 InFrameInterval);
+    static void Destroy();
+    
+    /** Access to the IOS Frame Pacer: CADisplayLink */
+    static FIOSFramePacer* FramePacer;
+    
+    /** Number of frames before the CADisplayLink triggers it's readied callback */
+    static uint32 FrameInterval;
+};
+
+
+typedef FIOSPlatformRHIFramePacer FPlatformRHIFramePacer;
\ No newline at end of file
diff --git a/Engine/Source/Runtime/Core/Public/IOS/IOSPlatformIncludes.h b/Engine/Source/Runtime/Core/Public/IOS/IOSPlatformIncludes.h
--- a/Engine/Source/Runtime/Core/Public/IOS/IOSPlatformIncludes.h
+++ b/Engine/Source/Runtime/Core/Public/IOS/IOSPlatformIncludes.h
@@ -97,6 +97,7 @@ struct RECT
 #include "IOS/IOSPlatformSurvey.h"
 #include "IOSAsyncTask.h"
 #include "IOS/IOSPlatformHttp.h"
+#include "IOS/IOSPlatformFramePacer.h"
 #import "IOSPlatformString.h"
 
 // include platform properties and typedef it for the runtime
diff --git a/Engine/Source/Runtime/Core/Public/Linux/LinuxPlatformIncludes.h b/Engine/Source/Runtime/Core/Public/Linux/LinuxPlatformIncludes.h
--- a/Engine/Source/Runtime/Core/Public/Linux/LinuxPlatformIncludes.h
+++ b/Engine/Source/Runtime/Core/Public/Linux/LinuxPlatformIncludes.h
@@ -69,6 +69,8 @@ struct RECT
 #include "Linux/LinuxPlatformSurvey.h"
 #include "Linux/LinuxPlatformHttp.h"
 
+typedef FGenericPlatformRHIFramePacer FPlatformRHIFramePacer;
+
 typedef FGenericPlatformAffinity FPlatformAffinity;
 
 // include platform properties and typedef it for the runtime
diff --git a/Engine/Source/Runtime/Core/Public/Mac/MacPlatformIncludes.h b/Engine/Source/Runtime/Core/Public/Mac/MacPlatformIncludes.h
--- a/Engine/Source/Runtime/Core/Public/Mac/MacPlatformIncludes.h
+++ b/Engine/Source/Runtime/Core/Public/Mac/MacPlatformIncludes.h
@@ -92,6 +92,8 @@ struct RECT
 #include "Mac/MacPlatformSurvey.h"
 #include "Mac/MacPlatformHttp.h"
 
+typedef FGenericPlatformRHIFramePacer FPlatformRHIFramePacer;
+
 typedef FGenericPlatformAffinity FPlatformAffinity;
 
 // include platform properties and typedef it for the runtime
diff --git a/Engine/Source/Runtime/Core/Public/WinRT/WinRTIncludes.h b/Engine/Source/Runtime/Core/Public/WinRT/WinRTIncludes.h
--- a/Engine/Source/Runtime/Core/Public/WinRT/WinRTIncludes.h
+++ b/Engine/Source/Runtime/Core/Public/WinRT/WinRTIncludes.h
@@ -24,6 +24,8 @@
 #include "WinRT/WinRTSurvey.h"
 #include "WinRT/WinRTHttp.h"
 
+typedef FGenericPlatformRHIFramePacer FPlatformRHIFramePacer;
+
 typedef FGenericPlatformAffinity FPlatformAffinity;
 
 // include platform properties and typedef it for the runtime
diff --git a/Engine/Source/Runtime/Core/Public/Windows/WindowsPlatformIncludes.h b/Engine/Source/Runtime/Core/Public/Windows/WindowsPlatformIncludes.h
--- a/Engine/Source/Runtime/Core/Public/Windows/WindowsPlatformIncludes.h
+++ b/Engine/Source/Runtime/Core/Public/Windows/WindowsPlatformIncludes.h
@@ -66,6 +66,8 @@ using namespace vc_attributes;
 #include "Windows/WindowsPlatformSurvey.h"
 #include "Windows/WindowsPlatformHttp.h"
 
+typedef FGenericPlatformRHIFramePacer FPlatformRHIFramePacer;
+
 typedef FGenericPlatformAffinity FPlatformAffinity;
 
 // include platform properties and typedef it for the runtime
diff --git a/Engine/Source/Runtime/IOS/MetalRHI/Private/MetalManager.cpp b/Engine/Source/Runtime/IOS/MetalRHI/Private/MetalManager.cpp
--- a/Engine/Source/Runtime/IOS/MetalRHI/Private/MetalManager.cpp
+++ b/Engine/Source/Runtime/IOS/MetalRHI/Private/MetalManager.cpp
@@ -145,43 +145,7 @@ void FMetalManager::ReleaseObject(id Object)
 	FMetalManager::Get()->DelayedFreeLists[FMetalManager::Get()->WhichFreeList].Add(Object);
 }
 
-
-@interface MetalFramePacer : NSObject
-{
-@public
-	FEvent *FramePacerEvent;
-}
-
--(void)run:(id)param;
--(void)signal:(id)param;
-
-@end
-
-@implementation MetalFramePacer
-
--(void)run:(id)param
-{
-	NSRunLoop *runloop = [NSRunLoop currentRunLoop];
-	CADisplayLink *displayLink = [CADisplayLink displayLinkWithTarget:self selector:@selector(signal:)];
-	displayLink.frameInterval = 2;
-	
-	[NSThread currentThread].name = @"FramePacerThread";
-	FramePacerEvent	= new FLocalEvent();
-	FramePacerEvent->Create(false);
-		
-	[displayLink addToRunLoop:runloop forMode:NSDefaultRunLoopMode];
-	[runloop run];
-}
-
--(void)signal:(id)param
-{
-	FramePacerEvent->Trigger();
-}
-
-@end
-
-MetalFramePacer *FramePacer;
-
+ 
 FMetalManager::FMetalManager()
 	: Device([IOSAppDelegate GetDelegate].IOSView->MetalDevice)
 	, CurrentCommandBuffer(nil)
@@ -228,10 +192,9 @@ FMetalManager::FMetalManager()
 
 	AutoReleasePoolTLSSlot = FPlatformTLS::AllocTlsSlot();
 
-	// Create display link thread
-	FramePacer = [[MetalFramePacer alloc] init];
-	[NSThread detachNewThreadSelector:@selector(run:) toTarget:FramePacer withObject:nil];
-	
+    FrameReadyEvent = FPlatformProcess::CreateSynchEvent();
+    FIOSPlatformRHIFramePacer::InitWithEvent( FrameReadyEvent, 2 );
+    
 	InitFrame();
 }
 
@@ -359,9 +322,9 @@ void FMetalManager::EndFrame(bool bPresent)
 	 {
 		dispatch_semaphore_signal(CommandBufferSemaphore);
 	 }];
-	
+    
 	// Wait until at least 2 VBlanks has passed since last time
-//	FramePacer->FramePacerEvent->Wait();
+    FrameReadyEvent->Wait();
 
 	// Commit before waiting to avoid leaving the gpu idle
 	[CurrentCommandBuffer commit];
diff --git a/Engine/Source/Runtime/IOS/MetalRHI/Private/MetalManager.h b/Engine/Source/Runtime/IOS/MetalRHI/Private/MetalManager.h
--- a/Engine/Source/Runtime/IOS/MetalRHI/Private/MetalManager.h
+++ b/Engine/Source/Runtime/IOS/MetalRHI/Private/MetalManager.h
@@ -244,6 +244,8 @@ protected:
 
 	void CommitNonComputeShaderConstants();
 
+private:
+    FEvent* FrameReadyEvent;
 };
 
 // Stats
diff --git a/Engine/Source/Runtime/IOS/MetalRHI/Private/MetalRHIPrivate.h b/Engine/Source/Runtime/IOS/MetalRHI/Private/MetalRHIPrivate.h
--- a/Engine/Source/Runtime/IOS/MetalRHI/Private/MetalRHIPrivate.h
+++ b/Engine/Source/Runtime/IOS/MetalRHI/Private/MetalRHIPrivate.h
@@ -47,221 +47,6 @@ extern TMap<id, int32> ClassCounts;
 
 
 
-
-
-
-class FLocalEvent : public FEvent
-{
-	// This is a little complicated, in an attempt to match Win32 Event semantics...
-	typedef enum
-	{
-		TRIGGERED_NONE,
-		TRIGGERED_ONE,
-		TRIGGERED_ALL,
-	} TriggerType;
-	
-	bool bInitialized;
-	bool bIsManualReset;
-	volatile TriggerType Triggered;
-	volatile int32 WaitingThreads;
-	pthread_mutex_t Mutex;
-	pthread_cond_t Condition;
-	
-	inline void LockEventMutex()
-	{
-		const int rc = pthread_mutex_lock(&Mutex);
-		check(rc == 0);
-	}
-	
-	inline void UnlockEventMutex()
-	{
-		const int rc = pthread_mutex_unlock(&Mutex);
-		check(rc == 0);
-	}
-	
-	static inline void SubtractTimevals(const struct timeval *FromThis, struct timeval *SubThis, struct timeval *Difference)
-	{
-		if (FromThis->tv_usec < SubThis->tv_usec)
-		{
-			int nsec = (SubThis->tv_usec - FromThis->tv_usec) / 1000000 + 1;
-			SubThis->tv_usec -= 1000000 * nsec;
-			SubThis->tv_sec += nsec;
-		}
-		
-		if (FromThis->tv_usec - SubThis->tv_usec > 1000000)
-		{
-			int nsec = (FromThis->tv_usec - SubThis->tv_usec) / 1000000;
-			SubThis->tv_usec += 1000000 * nsec;
-			SubThis->tv_sec -= nsec;
-		}
-		
-		Difference->tv_sec = FromThis->tv_sec - SubThis->tv_sec;
-		Difference->tv_usec = FromThis->tv_usec - SubThis->tv_usec;
-	}
-	
-public:
-	
-	FLocalEvent()
-	{
-		bInitialized = false;
-		bIsManualReset = false;
-		Triggered = TRIGGERED_NONE;
-		WaitingThreads = 0;
-	}
-	
-	virtual ~FLocalEvent()
-	{
-		// Safely destructing an Event is VERY delicate, so it can handle badly-designed
-		//  calling code that might still be waiting on the event.
-		if (bInitialized)
-		{
-			// try to flush out any waiting threads...
-			LockEventMutex();
-			bIsManualReset = true;
-			UnlockEventMutex();
-			Trigger();  // any waiting threads should start to wake up now.
-			
-			LockEventMutex();
-			bInitialized = false;  // further incoming calls to this object will now crash in check().
-			while (WaitingThreads)  // spin while waiting threads wake up.
-			{
-				UnlockEventMutex();  // cycle through waiting threads...
-				LockEventMutex();
-			}
-			// No threads are currently waiting on Condition and we hold the Mutex. Kill it.
-			pthread_cond_destroy(&Condition);
-			// Unlock and kill the mutex, since nothing else can grab it now.
-			UnlockEventMutex();
-			pthread_mutex_destroy(&Mutex);
-		}
-	}
-	
-	virtual bool Create(bool _bIsManualReset = false) override
-	{
-		check(!bInitialized);
-		bool RetVal = false;
-		Triggered = TRIGGERED_NONE;
-		bIsManualReset = _bIsManualReset;
-		
-		if (pthread_mutex_init(&Mutex, NULL) == 0)
-		{
-			if (pthread_cond_init(&Condition, NULL) == 0)
-			{
-				bInitialized = true;
-				RetVal = true;
-			}
-			else
-			{
-				pthread_mutex_destroy(&Mutex);
-			}
-		}
-		return RetVal;
-	}
-	
-	virtual void Trigger() override
-	{
-		check(bInitialized);
-		
-		LockEventMutex();
-		
-		if (bIsManualReset)
-		{
-			// Release all waiting threads at once.
-			Triggered = TRIGGERED_ALL;
-			int rc = pthread_cond_broadcast(&Condition);
-			check(rc == 0);
-		}
-		else
-		{
-			// Release one or more waiting threads (first one to get the mutex
-			//  will reset Triggered, rest will go back to waiting again).
-			Triggered = TRIGGERED_ONE;
-			int rc = pthread_cond_signal(&Condition);  // may release multiple threads anyhow!
-			check(rc == 0);
-		}
-		
-		UnlockEventMutex();
-	}
-	
-	virtual void Reset() override
-	{
-		check(bInitialized);
-		LockEventMutex();
-		Triggered = TRIGGERED_NONE;
-		UnlockEventMutex();
-	}
-	
-	virtual bool Wait(uint32 WaitTime, const bool bIgnoreThreadIdleStats = false) override
-	{
-		check(bInitialized);
-		
-		struct timeval StartTime;
-		
-		// We need to know the start time if we're going to do a timed wait.
-		if ( (WaitTime > 0) && (WaitTime != ((uint32)-1)) )  // not polling and not infinite wait.
-		{
-			gettimeofday(&StartTime, NULL);
-		}
-		
-		LockEventMutex();
-		
-		bool bRetVal = false;
-		
-		// loop in case we fall through the Condition signal but someone else claims the event.
-		do
-		{
-			// See what state the event is in...we may not have to wait at all...
-			
-			// One thread should be released. We saw it first, so we'll take it.
-			if (Triggered == TRIGGERED_ONE)
-			{
-				Triggered = TRIGGERED_NONE;  // dibs!
-				bRetVal = true;
-			}
-			
-			// manual reset that is still signaled. Every thread goes.
-			else if (Triggered == TRIGGERED_ALL)
-			{
-				bRetVal = true;
-			}
-			
-			// No event signalled yet.
-			else if (WaitTime != 0)  // not just polling, wait on the condition variable.
-			{
-				WaitingThreads++;
-				if (WaitTime == ((uint32)-1)) // infinite wait?
-				{
-					int rc = pthread_cond_wait(&Condition, &Mutex);  // unlocks Mutex while blocking...
-					check(rc == 0);
-				}
-				else  // timed wait.
-				{
-					struct timespec TimeOut;
-					const uint32 ms = (StartTime.tv_usec / 1000) + WaitTime;
-					TimeOut.tv_sec = StartTime.tv_sec + (ms / 1000);
-					TimeOut.tv_nsec = (ms % 1000) * 1000000;  // remainder of milliseconds converted to nanoseconds.
-					int rc = pthread_cond_timedwait(&Condition, &Mutex, &TimeOut);    // unlocks Mutex while blocking...
-					check((rc == 0) || (rc == ETIMEDOUT));
-					
-					// Update WaitTime and StartTime in case we have to go again...
-					struct timeval Now, Difference;
-					gettimeofday(&Now, NULL);
-					SubtractTimevals(&Now, &StartTime, &Difference);
-					const int32 DifferenceMS = ((Difference.tv_sec * 1000) + (Difference.tv_usec / 1000));
-					WaitTime = ((DifferenceMS >= WaitTime) ? 0 : (WaitTime - DifferenceMS));
-					StartTime = Now;
-				}
-				WaitingThreads--;
-				check(WaitingThreads >= 0);
-			}
-			
-		} while ((!bRetVal) && (WaitTime != 0));
-		
-		UnlockEventMutex();
-		return bRetVal;
-	}
-};
-
 FORCEINLINE int32 GetMetalCubeFace(ECubeFace Face)
 {
 	switch (Face)
diff --git a/Engine/Source/Runtime/OpenGLDrv/Private/IOS/IOSOpenGL.cpp b/Engine/Source/Runtime/OpenGLDrv/Private/IOS/IOSOpenGL.cpp
--- a/Engine/Source/Runtime/OpenGLDrv/Private/IOS/IOSOpenGL.cpp
+++ b/Engine/Source/Runtime/OpenGLDrv/Private/IOS/IOSOpenGL.cpp
@@ -115,8 +115,23 @@ void* PlatformGetWindow(FPlatformOpenGLContext* Context, void** AddParam)
 	return (void*)Context->Context;
 }
 
+
+// Event for coordinating pausing of render thread to keep inline with the ios display link.
+static FEvent* FrameReadyEvent = NULL;
+
 bool PlatformBlitToViewport( FPlatformOpenGLDevice* Device, const FOpenGLViewport& Viewport, uint32 BackbufferSizeX, uint32 BackbufferSizeY, bool bPresent,bool bLockToVsync, int32 SyncInterval )
 {
+    if( FIOSPlatformRHIFramePacer::IsEnabled() )
+    {
+        if( FrameReadyEvent == NULL )
+        {
+            FrameReadyEvent = FPlatformProcess::CreateSynchEvent();
+            FIOSPlatformRHIFramePacer::InitWithEvent( FrameReadyEvent, 2 );
+        }
+    
+        FrameReadyEvent->Wait();
+    }
+    
 	FPlatformOpenGLContext* const Context = Viewport.GetGLContext();
 
 	// @todo-mobile
