- Hopeful fix for tappy chicken squares

diff --git a/Engine/Source/Runtime/IOS/MetalRHI/Private/MetalManager.cpp b/Engine/Source/Runtime/IOS/MetalRHI/Private/MetalManager.cpp
--- a/Engine/Source/Runtime/IOS/MetalRHI/Private/MetalManager.cpp
+++ b/Engine/Source/Runtime/IOS/MetalRHI/Private/MetalManager.cpp
@@ -524,6 +524,12 @@ void FMetalManager::UpdateContext()
 		//@todo-rco: Do we need to test changes in Load/Store actions (and/or clear values) for Color & Depth?
 	}
 
+	// handle the case where going from backbuffer + depth -> backbuffer + null, no need to reset RT and do a store/load
+	if (CurrentNumRenderTargets == 1 && CurrentColorRenderTextures[0] == PreviousColorRenderTextures[0] && CurrentDepthRenderTexture == nil && CurrentStencilRenderTexture == PreviousStencilRenderTexture)
+	{
+		return;
+	}
+
 	PreviousNumRenderTargets = CurrentNumRenderTargets;
 
 	// if we are setting them to nothing, then this is probably end of frame, and we can't make a framebuffer
