UE4 - Metal compile fix

diff --git a/Engine/Source/Runtime/IOS/MetalRHI/Private/MetalManager.cpp b/Engine/Source/Runtime/IOS/MetalRHI/Private/MetalManager.cpp
--- a/Engine/Source/Runtime/IOS/MetalRHI/Private/MetalManager.cpp
+++ b/Engine/Source/Runtime/IOS/MetalRHI/Private/MetalManager.cpp
@@ -525,7 +525,7 @@ void FMetalManager::UpdateContext()
 	}
 
 	// handle the case where going from backbuffer + depth -> backbuffer + null, no need to reset RT and do a store/load
-	if (CurrentNumRenderTargets == 1 && CurrentColorRenderTextures[0] == PreviousColorRenderTextures[0] && CurrentDepthRenderTexture == nil && CurrentStencilRenderTexture == PreviousStencilRenderTexture)
+	if (CurrentNumRenderTargets == 1 && CurrentColorRenderTextures[0] == PreviousColorRenderTextures[0] && CurrentDepthRenderTexture == nil)
 	{
 		return;
 	}
