Tappy Framerate no longer locked at 60.

diff --git a/Engine/Source/Runtime/IOS/MetalRHI/Private/MetalManager.cpp b/Engine/Source/Runtime/IOS/MetalRHI/Private/MetalManager.cpp
--- a/Engine/Source/Runtime/IOS/MetalRHI/Private/MetalManager.cpp
+++ b/Engine/Source/Runtime/IOS/MetalRHI/Private/MetalManager.cpp
@@ -193,7 +193,7 @@ FMetalManager::FMetalManager()
 	AutoReleasePoolTLSSlot = FPlatformTLS::AllocTlsSlot();
 
     FrameReadyEvent = FPlatformProcess::CreateSynchEvent();
-    FIOSPlatformRHIFramePacer::InitWithEvent( FrameReadyEvent, 2 );
+    FIOSPlatformRHIFramePacer::InitWithEvent( FrameReadyEvent, 1 );
     
 	InitFrame();
 }
@@ -323,7 +323,7 @@ void FMetalManager::EndFrame(bool bPresent)
 		dispatch_semaphore_signal(CommandBufferSemaphore);
 	 }];
     
-	// Wait until at least 2 VBlanks has passed since last time
+	// Wait until at least a VBlank has passed since last time
     FrameReadyEvent->Wait();
 
 	// Commit before waiting to avoid leaving the gpu idle
@@ -885,4 +885,4 @@ void FMetalManager::SetRasterizerState(const FRasterizerStateInitializerRHI& Sta
 #endif
 
 	bFirstRasterizerState = false;
-}
\ No newline at end of file
+}
diff --git a/Engine/Source/Runtime/OpenGLDrv/Private/IOS/IOSOpenGL.cpp b/Engine/Source/Runtime/OpenGLDrv/Private/IOS/IOSOpenGL.cpp
--- a/Engine/Source/Runtime/OpenGLDrv/Private/IOS/IOSOpenGL.cpp
+++ b/Engine/Source/Runtime/OpenGLDrv/Private/IOS/IOSOpenGL.cpp
@@ -126,7 +126,7 @@ bool PlatformBlitToViewport( FPlatformOpenGLDevice* Device, const FOpenGLViewpor
         if( FrameReadyEvent == NULL )
         {
             FrameReadyEvent = FPlatformProcess::CreateSynchEvent();
-            FIOSPlatformRHIFramePacer::InitWithEvent( FrameReadyEvent, 2 );
+            FIOSPlatformRHIFramePacer::InitWithEvent( FrameReadyEvent, 1 );
         }
     
         FrameReadyEvent->Wait();
