- Added METAL_MRT shader format (part one of a multi-checkin for MetalMRT support), but it's not usable yet.

diff --git a/Engine/Shaders/Common.usf b/Engine/Shaders/Common.usf
--- a/Engine/Shaders/Common.usf
+++ b/Engine/Shaders/Common.usf
@@ -57,8 +57,7 @@
 	#define MaterialFloat4x3 float4x3 
 #endif
 
-//@todo-rco: Temp workaround for Seed 2 & 3
-#if COMPILER_GLSL_ES2 || METAL_PROFILE
+#if COMPILER_GLSL_ES2 || METAL_PROFILE || METAL_MRT_PROFILE
 
 	float determinant(float3x3 M)
 	{    
@@ -120,6 +119,9 @@
 #elif METAL_PROFILE
 	#define FEATURE_LEVEL FEATURE_LEVEL_ES2
 	#define FCOLOR_COMPONENT_SWIZZLE .bgra  
+#elif METAL_MRT_PROFILE
+	#define FEATURE_LEVEL FEATURE_LEVEL_SM4
+	#define FCOLOR_COMPONENT_SWIZZLE .bgra  
 #elif ES2_PROFILE || ES3_1_PROFILE
 	//@todo ES3_1 GL
 	#define FEATURE_LEVEL FEATURE_LEVEL_ES2
diff --git a/Engine/Shaders/DeferredShadingCommon.usf b/Engine/Shaders/DeferredShadingCommon.usf
--- a/Engine/Shaders/DeferredShadingCommon.usf
+++ b/Engine/Shaders/DeferredShadingCommon.usf
@@ -565,6 +565,21 @@ SamplerState CustomDepthTextureSampler;
 // @param UV - UV space in the GBuffer textures (BufferSize resolution)
 FGBufferData GetGBufferData(float2 UV, bool bGetNormalizedNormal = true)
 {
+#if 0 //METAL_MRT_PROFILE
+	// @todo metal mrt: The deferred renderer isn't keeping these in tiled memory all the time - we don't know when this makes sense
+	// versus just sampling a bound resolved texture
+	float4 GBufferA = FramebufferFetchMRT(1);
+	float4 GBufferB = FramebufferFetchMRT(2);
+	float4 GBufferC = FramebufferFetchMRT(3);
+	float4 GBufferD = FramebufferFetchMRT(4);
+	// @todo metal mrt: We can't framebuffer fetch the depth, can we jam it in somewhere?
+	float CustomNativeDepth = 0.5;
+	#if ALLOW_STATIC_LIGHTING
+		float4 GBufferE = FramebufferFetchMRT(5);
+	#else
+		float4 GBufferE = 1;
+	#endif
+#else	
 	float4 GBufferA = Texture2DSampleLevel(GBuffers.GBufferATexture, GBuffers.GBufferATextureSampler, UV, 0);
 	float4 GBufferB = Texture2DSampleLevel(GBuffers.GBufferBTexture, GBuffers.GBufferBTextureSampler, UV, 0);
 	float4 GBufferC = Texture2DSampleLevel(GBuffers.GBufferCTexture, GBuffers.GBufferCTextureSampler, UV, 0);
@@ -576,6 +591,7 @@ FGBufferData GetGBufferData(float2 UV, bool bGetNormalizedNormal = true)
 	#else
 		float4 GBufferE = 1;
 	#endif
+#endif
 
 	float SceneDepth = CalcSceneDepth(UV);
 	
diff --git a/Engine/Source/Developer/iOS/IOSTargetPlatform/Private/IOSTargetPlatform.cpp b/Engine/Source/Developer/iOS/IOSTargetPlatform/Private/IOSTargetPlatform.cpp
--- a/Engine/Source/Developer/iOS/IOSTargetPlatform/Private/IOSTargetPlatform.cpp
+++ b/Engine/Source/Developer/iOS/IOSTargetPlatform/Private/IOSTargetPlatform.cpp
@@ -294,6 +294,7 @@ void FIOSTargetPlatform::GetAllPossibleShaderFormats( TArray<FName>& OutFormats
 {
 	static FName NAME_OPENGL_ES2_IOS(TEXT("GLSL_ES2_IOS"));
 	static FName NAME_SF_METAL(TEXT("SF_METAL"));
+	static FName NAME_SF_METAL_MRT(TEXT("SF_METAL_MRT"));
 
 	// default to supporting ES2
 	bool bSupportsOpenGLES2 = true;
@@ -309,6 +310,8 @@ void FIOSTargetPlatform::GetAllPossibleShaderFormats( TArray<FName>& OutFormats
 	if (bSupportsMetal)
 	{
 		OutFormats.AddUnique(NAME_SF_METAL);
+		// @todo metal mrt: Add a project setting to support MetalMRT mode
+		// OutFormats.AddUnique(NAME_SF_METAL_MRT);
 	}
 }
 
diff --git a/Engine/Source/Developer/iOS/MetalShaderFormat/Private/MetalShaderCompiler.cpp b/Engine/Source/Developer/iOS/MetalShaderFormat/Private/MetalShaderCompiler.cpp
--- a/Engine/Source/Developer/iOS/MetalShaderFormat/Private/MetalShaderCompiler.cpp
+++ b/Engine/Source/Developer/iOS/MetalShaderFormat/Private/MetalShaderCompiler.cpp
@@ -820,9 +820,27 @@ void CompileShader_Metal(const FShaderCompilerInput& Input,FShaderCompilerOutput
 	FShaderCompilerDefinitions AdditionalDefines;
 	EHlslCompileTarget HlslCompilerTarget = HCT_FeatureLevelES3_1;
 	AdditionalDefines.SetDefine(TEXT("IOS"), 1);
+	AdditionalDefines.SetDefine(TEXT("row_major"), TEXT(""));
+
+	static FName NAME_SF_METAL(TEXT("SF_METAL"));
+	static FName NAME_SF_METAL_MRT(TEXT("SF_METAL_MRT"));
+
+	if (Input.ShaderFormat == NAME_SF_METAL)
+	{
 	AdditionalDefines.SetDefine(TEXT("METAL_PROFILE"), 1);
 	AdditionalDefines.SetDefine(TEXT("ES3_1_PROFILE"), 1);
-	AdditionalDefines.SetDefine(TEXT("row_major"), TEXT(""));
+	}
+	else if (Input.ShaderFormat == NAME_SF_METAL_MRT)
+	{
+		AdditionalDefines.SetDefine(TEXT("METAL_MRT_PROFILE"), 1);
+		AdditionalDefines.SetDefine(TEXT("SM4_PROFILE"), 1);
+	}
+	else
+	{
+		Output.bSucceeded = false;
+		new(Output.Errors) FShaderCompilerError(*FString::Printf(TEXT("Invalid shader format '%s' passed to compiler."), *Input.ShaderFormat.ToString()));
+		return;
+	}
 	
 	const bool bDumpDebugInfo = (Input.DumpDebugInfoPath != TEXT("") && IFileManager::Get().DirectoryExists(*Input.DumpDebugInfoPath));
 
diff --git a/Engine/Source/Developer/iOS/MetalShaderFormat/Private/MetalShaderFormat.cpp b/Engine/Source/Developer/iOS/MetalShaderFormat/Private/MetalShaderFormat.cpp
--- a/Engine/Source/Developer/iOS/MetalShaderFormat/Private/MetalShaderFormat.cpp
+++ b/Engine/Source/Developer/iOS/MetalShaderFormat/Private/MetalShaderFormat.cpp
@@ -9,6 +9,7 @@
 #include "hlslcc.h"
 
 static FName NAME_SF_METAL(TEXT("SF_METAL"));
+static FName NAME_SF_METAL_MRT(TEXT("SF_METAL_MRT"));
 
 class FMetalShaderFormat : public IShaderFormat
 {
@@ -49,10 +50,11 @@ public:
 	virtual void GetSupportedFormats(TArray<FName>& OutFormats) const
 	{
 		OutFormats.Add(NAME_SF_METAL);
+		OutFormats.Add(NAME_SF_METAL_MRT);
 	}
 	virtual void CompileShader(FName Format, const struct FShaderCompilerInput& Input, struct FShaderCompilerOutput& Output,const FString& WorkingDirectory) const
 	{
-		check(Format == NAME_SF_METAL);
+		check(Format == NAME_SF_METAL || Format == NAME_SF_METAL_MRT);
 		CompileShader_Metal(Input, Output, WorkingDirectory);
 	}
 };
diff --git a/Engine/Source/Runtime/IOS/MetalRHI/Private/MetalCommands.cpp b/Engine/Source/Runtime/IOS/MetalRHI/Private/MetalCommands.cpp
--- a/Engine/Source/Runtime/IOS/MetalRHI/Private/MetalCommands.cpp
+++ b/Engine/Source/Runtime/IOS/MetalRHI/Private/MetalCommands.cpp
@@ -21,9 +21,12 @@ static MTLPrimitiveType TranslatePrimitiveType(uint32 PrimitiveType)
 
 void FMetalDynamicRHI::RHISetStreamSource(uint32 StreamIndex,FVertexBufferRHIParamRef VertexBufferRHI,uint32 Stride,uint32 Offset)
 {
-	DYNAMIC_CAST_METGALRESOURCE(VertexBuffer,VertexBuffer);
+	DYNAMIC_CAST_METALRESOURCE(VertexBuffer,VertexBuffer);
 
-	[FMetalManager::GetContext() setVertexBuffer:VertexBuffer->Buffer offset:VertexBuffer->Offset + Offset atIndex:UNREAL_TO_METAL_BUFFER_INDEX(StreamIndex)];
+	if (VertexBuffer != NULL)
+	{
+		[FMetalManager::GetContext() setVertexBuffer:VertexBuffer->Buffer offset:VertexBuffer->Offset + Offset atIndex:UNREAL_TO_METAL_BUFFER_INDEX(StreamIndex)];
+	}
 }
 
 void FMetalDynamicRHI::RHISetStreamOutTargets(uint32 NumTargets, const FVertexBufferRHIParamRef* VertexBuffers, const uint32* Offsets)
@@ -34,7 +37,7 @@ void FMetalDynamicRHI::RHISetStreamOutTargets(uint32 NumTargets, const FVertexBu
 
 void FMetalDynamicRHI::RHISetRasterizerState(FRasterizerStateRHIParamRef NewStateRHI)
 {
-	DYNAMIC_CAST_METGALRESOURCE(RasterizerState,NewState);
+	DYNAMIC_CAST_METALRESOURCE(RasterizerState,NewState);
 
 	NewState->Set();
 }
@@ -90,7 +93,7 @@ void FMetalDynamicRHI::RHISetScissorRect(bool bEnable,uint32 MinX,uint32 MinY,ui
 
 void FMetalDynamicRHI::RHISetBoundShaderState( FBoundShaderStateRHIParamRef BoundShaderStateRHI)
 {
-	DYNAMIC_CAST_METGALRESOURCE(BoundShaderState,BoundShaderState);
+	DYNAMIC_CAST_METALRESOURCE(BoundShaderState,BoundShaderState);
 
 	FMetalManager::Get()->SetBoundShaderState(BoundShaderState);
 	BoundShaderStateHistory.Add(BoundShaderState);
@@ -164,7 +167,10 @@ void FMetalDynamicRHI::RHISetShaderResourceViewParameter(FGeometryShaderRHIParam
 
 void FMetalDynamicRHI::RHISetShaderResourceViewParameter(FPixelShaderRHIParamRef PixelShaderRHI,uint32 TextureIndex,FShaderResourceViewRHIParamRef SRVRHI)
 {
-	NOT_SUPPORTED("RHISetShaderResourceViewParameter");
+	DYNAMIC_CAST_METALRESOURCE(ShaderResourceView, SRV);
+
+	FMetalSurface& Surface = GetMetalSurfaceFromRHITexture(SRV->SourceTexture.GetReference());
+	[FMetalManager::GetContext() setFragmentTexture:Surface.Texture atIndex:TextureIndex];
 }
 
 void FMetalDynamicRHI::RHISetShaderResourceViewParameter(FComputeShaderRHIParamRef ComputeShaderRHI,uint32 TextureIndex,FShaderResourceViewRHIParamRef SRVRHI)
@@ -175,42 +181,42 @@ void FMetalDynamicRHI::RHISetShaderResourceViewParameter(FComputeShaderRHIParamR
 
 void FMetalDynamicRHI::RHISetShaderSampler(FVertexShaderRHIParamRef VertexShaderRHI, uint32 SamplerIndex, FSamplerStateRHIParamRef NewStateRHI)
 {
-	DYNAMIC_CAST_METGALRESOURCE(SamplerState,NewState);
+	DYNAMIC_CAST_METALRESOURCE(SamplerState,NewState);
 
 	[FMetalManager::GetContext() setVertexSamplerState:NewState->State atIndex:SamplerIndex];
 }
 
 void FMetalDynamicRHI::RHISetShaderSampler(FHullShaderRHIParamRef HullShader, uint32 SamplerIndex, FSamplerStateRHIParamRef NewStateRHI)
 {
-	DYNAMIC_CAST_METGALRESOURCE(SamplerState,NewState);
+	DYNAMIC_CAST_METALRESOURCE(SamplerState,NewState);
 
 	NOT_SUPPORTED("RHISetSamplerState-Hull");
 }
 
 void FMetalDynamicRHI::RHISetShaderSampler(FDomainShaderRHIParamRef DomainShader, uint32 SamplerIndex, FSamplerStateRHIParamRef NewStateRHI)
 {
-	DYNAMIC_CAST_METGALRESOURCE(SamplerState,NewState);
+	DYNAMIC_CAST_METALRESOURCE(SamplerState,NewState);
 
 	NOT_SUPPORTED("RHISetSamplerState-Domain");
 }
 
 void FMetalDynamicRHI::RHISetShaderSampler(FGeometryShaderRHIParamRef GeometryShader, uint32 SamplerIndex, FSamplerStateRHIParamRef NewStateRHI)
 {
-	DYNAMIC_CAST_METGALRESOURCE(SamplerState,NewState);
+	DYNAMIC_CAST_METALRESOURCE(SamplerState,NewState);
 
 	NOT_SUPPORTED("RHISetSamplerState-Geometry");
 }
 
 void FMetalDynamicRHI::RHISetShaderSampler(FPixelShaderRHIParamRef PixelShader, uint32 SamplerIndex, FSamplerStateRHIParamRef NewStateRHI)
 {
-	DYNAMIC_CAST_METGALRESOURCE(SamplerState,NewState);
+	DYNAMIC_CAST_METALRESOURCE(SamplerState,NewState);
 
 	[FMetalManager::GetContext() setFragmentSamplerState:NewState->State atIndex:SamplerIndex];
 }
 
 void FMetalDynamicRHI::RHISetShaderSampler(FComputeShaderRHIParamRef ComputeShader, uint32 SamplerIndex, FSamplerStateRHIParamRef NewStateRHI)
 {
-	DYNAMIC_CAST_METGALRESOURCE(SamplerState,NewState);
+	DYNAMIC_CAST_METALRESOURCE(SamplerState,NewState);
 
 	NOT_SUPPORTED("RHISetSamplerState-Compute");
 }
@@ -251,7 +257,7 @@ void FMetalDynamicRHI::RHISetShaderParameter(FComputeShaderRHIParamRef ComputeSh
 
 void FMetalDynamicRHI::RHISetShaderUniformBuffer(FVertexShaderRHIParamRef VertexShaderRHI, uint32 BufferIndex, FUniformBufferRHIParamRef BufferRHI)
 {
-	DYNAMIC_CAST_METGALRESOURCE(VertexShader, VertexShader);
+	DYNAMIC_CAST_METALRESOURCE(VertexShader, VertexShader);
 	VertexShader->BoundUniformBuffers[BufferIndex] = BufferRHI;
 	VertexShader->DirtyUniformBuffers |= 1 << BufferIndex;
 
@@ -281,7 +287,7 @@ void FMetalDynamicRHI::RHISetShaderUniformBuffer(FGeometryShaderRHIParamRef Geom
 
 void FMetalDynamicRHI::RHISetShaderUniformBuffer(FPixelShaderRHIParamRef PixelShaderRHI, uint32 BufferIndex, FUniformBufferRHIParamRef BufferRHI)
 {
-	DYNAMIC_CAST_METGALRESOURCE(PixelShader, PixelShader);
+	DYNAMIC_CAST_METALRESOURCE(PixelShader, PixelShader);
 	PixelShader->BoundUniformBuffers[BufferIndex] = BufferRHI;
 	PixelShader->DirtyUniformBuffers |= 1 << BufferIndex;
 
@@ -302,14 +308,14 @@ void FMetalDynamicRHI::RHISetShaderUniformBuffer(FComputeShaderRHIParamRef Compu
 
 void FMetalDynamicRHI::RHISetDepthStencilState(FDepthStencilStateRHIParamRef NewStateRHI, uint32 StencilRef)
 {
-	DYNAMIC_CAST_METGALRESOURCE(DepthStencilState,NewState);
+	DYNAMIC_CAST_METALRESOURCE(DepthStencilState,NewState);
 
 	NewState->Set();
 }
 
 void FMetalDynamicRHI::RHISetBlendState(FBlendStateRHIParamRef NewStateRHI, const FLinearColor& BlendFactor)
 {
-	DYNAMIC_CAST_METGALRESOURCE(BlendState,NewState);
+	DYNAMIC_CAST_METALRESOURCE(BlendState,NewState);
 
 	NewState->Set();
 }
@@ -319,45 +325,10 @@ void FMetalDynamicRHI::RHISetRenderTargets(uint32 NumSimultaneousRenderTargets,
 	FTextureRHIParamRef NewDepthStencilTargetRHI, uint32 NumUAVs, const FUnorderedAccessViewRHIParamRef* UAVs)
 {
 	FMetalManager* Manager = FMetalManager::Get();
-	for (int32 RenderTargetIndex = 0; RenderTargetIndex < MaxMetalRenderTargets; RenderTargetIndex++)
-	{
-		const FRHIRenderTargetView& RenderTargetView = NewRenderTargets[RenderTargetIndex];
-		// update the current RTs
-		if (RenderTargetIndex < NumSimultaneousRenderTargets && RenderTargetView.Texture != NULL)
-		{
-			FMetalSurface& Surface = GetMetalSurfaceFromRHITexture(RenderTargetView.Texture);
-			Manager->SetCurrentRenderTarget(&Surface, RenderTargetIndex, RenderTargetView.MipIndex, RenderTargetView.ArraySliceIndex, GetMetalRTLoadAction(RenderTargetView.LoadAction), GetMetalRTStoreAction(RenderTargetView.StoreAction), NumSimultaneousRenderTargets);
-		}
-		else
-		{
-			Manager->SetCurrentRenderTarget(NULL, RenderTargetIndex, 0, 0, MTLLoadActionDontCare, MTLStoreActionStore, NumSimultaneousRenderTargets);
-		}
-	}
-	
-	if (NewDepthStencilTargetRHI)
-	{
-		FMetalSurface& Surface = GetMetalSurfaceFromRHITexture(NewDepthStencilTargetRHI);
-		Manager->SetCurrentDepthStencilTarget(&Surface, MTLLoadActionClear, MTLStoreActionDontCare, GUsesInvertedZ ? 0.0f : 1.0f);
-	}
-	else
-	{
-		Manager->SetCurrentDepthStencilTarget(NULL, MTLLoadActionClear, MTLStoreActionDontCare, GUsesInvertedZ ? 0.0f : 1.0f);
-	}
-	
-	// now that we have a new render target, we need a new context to render to it!
-	Manager->UpdateContext();
-	
-	// Set the viewport to the full size of render target 0.
-	if (NewRenderTargets && NewRenderTargets[0].Texture)
-	{
-		const FRHIRenderTargetView& RenderTargetView = NewRenderTargets[0];
-		FMetalSurface& RenderTarget = GetMetalSurfaceFromRHITexture(RenderTargetView.Texture);
-		
-		uint32 Width = FMath::Max(RenderTarget.Texture.width >> RenderTargetView.MipIndex, (uint32)1);
-		uint32 Height = FMath::Max(RenderTarget.Texture.height >> RenderTargetView.MipIndex, (uint32)1);
-		
-		RHISetViewport(0, 0, 0.0f, Width, Height, 1.0f);
-	}
+
+	FRHIDepthRenderTargetView DepthView(NewDepthStencilTargetRHI, ERenderTargetLoadAction::EClear, ERenderTargetStoreAction::ENoAction);
+	FRHISetRenderTargetsInfo Info(NumSimultaneousRenderTargets, NewRenderTargets, DepthView);
+	RHISetRenderTargetsAndClear(Info);
 }
 
 void FMetalDynamicRHI::RHIDiscardRenderTargets(bool Depth, bool Stencil, uint32 ColorBitMask)
@@ -367,6 +338,10 @@ void FMetalDynamicRHI::RHIDiscardRenderTargets(bool Depth, bool Stencil, uint32
 void FMetalDynamicRHI::RHISetRenderTargetsAndClear(const FRHISetRenderTargetsInfo& RenderTargetsInfo)
 {
 	FMetalManager* Manager = FMetalManager::Get();
+
+#if 1
+	Manager->SetRenderTargetsInfo(RenderTargetsInfo);
+#else
 	for (int32 RenderTargetIndex = 0; RenderTargetIndex < MaxMetalRenderTargets; RenderTargetIndex++)
 	{
 		const FRHIRenderTargetView& RenderTargetView = RenderTargetsInfo.ColorRenderTarget[RenderTargetIndex];
@@ -385,15 +360,23 @@ void FMetalDynamicRHI::RHISetRenderTargetsAndClear(const FRHISetRenderTargetsInf
 	if (RenderTargetsInfo.DepthStencilRenderTarget.Texture)
 	{
 		FMetalSurface& Surface = GetMetalSurfaceFromRHITexture(RenderTargetsInfo.DepthStencilRenderTarget.Texture);
-		Manager->SetCurrentDepthStencilTarget(&Surface, GetMetalRTLoadAction(RenderTargetsInfo.DepthStencilRenderTarget.LoadAction), GetMetalRTStoreAction(RenderTargetsInfo.DepthStencilRenderTarget.StoreAction), RenderTargetsInfo.DepthClearValue);
+		Manager->SetCurrentDepthStencilTarget(&Surface, 
+			GetMetalRTLoadAction(RenderTargetsInfo.DepthStencilRenderTarget.DepthLoadAction), 
+			GetMetalRTStoreAction(RenderTargetsInfo.DepthStencilRenderTarget.DepthStoreAction), 
+			RenderTargetsInfo.DepthClearValue,
+			GetMetalRTLoadAction(RenderTargetsInfo.DepthStencilRenderTarget.StencilLoadAction),
+			GetMetalRTStoreAction(RenderTargetsInfo.DepthStencilRenderTarget.StencilStoreAction),
+			RenderTargetsInfo.StencilClearValue);
 	}
 	else
 	{
-		Manager->SetCurrentDepthStencilTarget(NULL, MTLLoadActionClear, MTLStoreActionDontCare, 1.0f);
+		Manager->SetCurrentDepthStencilTarget(NULL);
 	}
 
 	// now that we have a new render target, we need a new context to render to it!
 	Manager->UpdateContext();
+#endif
+
 
 	// Set the viewport to the full size of render target 0.
 	if (RenderTargetsInfo.ColorRenderTarget[0].Texture)
@@ -411,14 +394,14 @@ void FMetalDynamicRHI::RHISetRenderTargetsAndClear(const FRHISetRenderTargetsInf
 // Occlusion/Timer queries.
 void FMetalDynamicRHI::RHIBeginRenderQuery(FRenderQueryRHIParamRef QueryRHI)
 {
-	DYNAMIC_CAST_METGALRESOURCE(RenderQuery,Query);
+	DYNAMIC_CAST_METALRESOURCE(RenderQuery,Query);
 
 	Query->Begin();
 }
 
 void FMetalDynamicRHI::RHIEndRenderQuery(FRenderQueryRHIParamRef QueryRHI)
 {
-	DYNAMIC_CAST_METGALRESOURCE(RenderQuery,Query);
+	DYNAMIC_CAST_METALRESOURCE(RenderQuery,Query);
 
 	Query->End();
 }
@@ -427,7 +410,7 @@ void FMetalDynamicRHI::RHIEndRenderQuery(FRenderQueryRHIParamRef QueryRHI)
 void FMetalDynamicRHI::RHIDrawPrimitive(uint32 PrimitiveType, uint32 BaseVertexIndex, uint32 NumPrimitives, uint32 NumInstances)
 {
 	SCOPE_CYCLE_COUNTER(STAT_MetalDrawCallTime);
-	checkf(NumInstances == 1, TEXT("Currently only 1 instance is supported"));
+	//checkf(NumInstances == 1, TEXT("Currently only 1 instance is supported"));
 	
 	RHI_DRAW_CALL_STATS(PrimitiveType,NumInstances*NumPrimitives);
 
@@ -454,12 +437,12 @@ void FMetalDynamicRHI::RHIDrawIndexedPrimitive(FIndexBufferRHIParamRef IndexBuff
 	uint32 NumVertices, uint32 StartIndex, uint32 NumPrimitives, uint32 NumInstances)
 {
 	SCOPE_CYCLE_COUNTER(STAT_MetalDrawCallTime);
-	checkf(NumInstances == 1, TEXT("Currently only 1 instance is supported"));
+	//checkf(NumInstances == 1, TEXT("Currently only 1 instance is supported"));
 	checkf(BaseVertexIndex  == 0, TEXT("BaseVertexIndex must be 0, see GRHISupportsBaseVertexIndex"));
 
 	RHI_DRAW_CALL_STATS(PrimitiveType,NumInstances*NumPrimitives);
 
-	DYNAMIC_CAST_METGALRESOURCE(IndexBuffer,IndexBuffer);
+	DYNAMIC_CAST_METALRESOURCE(IndexBuffer,IndexBuffer);
 
 	// finalize any pending state
 	FMetalManager::Get()->PrepareToDraw(NumVertices);
@@ -592,7 +575,7 @@ void FMetalDynamicRHI::RHIClear(bool bClearColor,const FLinearColor& Color, bool
 
 void FMetalDynamicRHI::RHIClearMRT(bool bClearColor,int32 NumClearColors,const FLinearColor* ClearColorArray,bool bClearDepth,float Depth,bool bClearStencil,uint32 Stencil, FIntRect ExcludeRect)
 {
-	checkf(NumClearColors == 1, TEXT("Only one MRT supported in RHIClearMRT"));
+	// this needs to be handled in SetRenderTarget LoadAction 
 }
 
 
diff --git a/Engine/Source/Runtime/IOS/MetalRHI/Private/MetalIndexBuffer.cpp b/Engine/Source/Runtime/IOS/MetalRHI/Private/MetalIndexBuffer.cpp
--- a/Engine/Source/Runtime/IOS/MetalRHI/Private/MetalIndexBuffer.cpp
+++ b/Engine/Source/Runtime/IOS/MetalRHI/Private/MetalIndexBuffer.cpp
@@ -72,14 +72,14 @@ FIndexBufferRHIRef FMetalDynamicRHI::RHICreateIndexBuffer(uint32 Stride,uint32 S
 
 void* FMetalDynamicRHI::RHILockIndexBuffer(FIndexBufferRHIParamRef IndexBufferRHI, uint32 Offset, uint32 Size, EResourceLockMode LockMode)
 {
-	DYNAMIC_CAST_METGALRESOURCE(IndexBuffer,IndexBuffer);
+	DYNAMIC_CAST_METALRESOURCE(IndexBuffer,IndexBuffer);
 
 	return (uint8*)IndexBuffer->Lock(LockMode, Size) + Offset;
 }
 
 void FMetalDynamicRHI::RHIUnlockIndexBuffer(FIndexBufferRHIParamRef IndexBufferRHI)
 {
-	DYNAMIC_CAST_METGALRESOURCE(IndexBuffer,IndexBuffer);
+	DYNAMIC_CAST_METALRESOURCE(IndexBuffer,IndexBuffer);
 
 	IndexBuffer->Unlock();
 }
diff --git a/Engine/Source/Runtime/IOS/MetalRHI/Private/MetalManager.cpp b/Engine/Source/Runtime/IOS/MetalRHI/Private/MetalManager.cpp
--- a/Engine/Source/Runtime/IOS/MetalRHI/Private/MetalManager.cpp
+++ b/Engine/Source/Runtime/IOS/MetalRHI/Private/MetalManager.cpp
@@ -9,25 +9,34 @@ const uint32 RingBufferSize = 8 * 1024 * 1024;
 TMap<id, int32> ClassCounts;
 #endif
 
-#define NUMBITS_BLEND_STATE				5
-#define NUMBITS_RENDER_TARGET_FORMAT	8
-#define NUMBITS_DEPTH_TARGET_FORMAT		8
-#define NUMBITS_SAMPLE_COUNT			3
+#define NUMBITS_BLEND_STATE				5 //(x6=30)
+#define NUMBITS_RENDER_TARGET_FORMAT	4 //(x6=24)
+#define NUMBITS_DEPTH_ENABLED			1 //(x1=1)
+#define NUMBITS_STENCIL_ENABLED			1 //(x1=1)
+#define NUMBITS_SAMPLE_COUNT			3 //(x1=3)
 
 
 #define OFFSET_BLEND_STATE0				(0)
-#define OFFSET_BLEND_STATE1				(OFFSET_BLEND_STATE0	+ NUMBITS_BLEND_STATE)
-#define OFFSET_BLEND_STATE2				(OFFSET_BLEND_STATE1	+ NUMBITS_BLEND_STATE)
-#define OFFSET_BLEND_STATE3				(OFFSET_BLEND_STATE2	+ NUMBITS_BLEND_STATE)
-#define OFFSET_RENDER_TARGET_FORMAT0	(OFFSET_BLEND_STATE3	+ NUMBITS_BLEND_STATE)
+#define OFFSET_BLEND_STATE1				(OFFSET_BLEND_STATE0			+ NUMBITS_BLEND_STATE)
+#define OFFSET_BLEND_STATE2				(OFFSET_BLEND_STATE1			+ NUMBITS_BLEND_STATE)
+#define OFFSET_BLEND_STATE3				(OFFSET_BLEND_STATE2			+ NUMBITS_BLEND_STATE)
+#define OFFSET_BLEND_STATE4				(OFFSET_BLEND_STATE3			+ NUMBITS_BLEND_STATE)
+#define OFFSET_BLEND_STATE5				(OFFSET_BLEND_STATE4			+ NUMBITS_BLEND_STATE)
+#define OFFSET_RENDER_TARGET_FORMAT0	(OFFSET_BLEND_STATE5			+ NUMBITS_BLEND_STATE)
 #define OFFSET_RENDER_TARGET_FORMAT1	(OFFSET_RENDER_TARGET_FORMAT0	+ NUMBITS_RENDER_TARGET_FORMAT)
 #define OFFSET_RENDER_TARGET_FORMAT2	(OFFSET_RENDER_TARGET_FORMAT1	+ NUMBITS_RENDER_TARGET_FORMAT)
 #define OFFSET_RENDER_TARGET_FORMAT3	(OFFSET_RENDER_TARGET_FORMAT2	+ NUMBITS_RENDER_TARGET_FORMAT)
-#define OFFSET_DEPTH_TARGET_FORMAT		(OFFSET_RENDER_TARGET_FORMAT3	+ NUMBITS_RENDER_TARGET_FORMAT)
-#define OFFSET_SAMPLE_COUNT				(OFFSET_DEPTH_TARGET_FORMAT		+ NUMBITS_DEPTH_TARGET_FORMAT)
+#define OFFSET_RENDER_TARGET_FORMAT4	(OFFSET_RENDER_TARGET_FORMAT3	+ NUMBITS_RENDER_TARGET_FORMAT)
+#define OFFSET_RENDER_TARGET_FORMAT5	(OFFSET_RENDER_TARGET_FORMAT4	+ NUMBITS_RENDER_TARGET_FORMAT)
+#define OFFSET_DEPTH_ENABLED			(OFFSET_RENDER_TARGET_FORMAT5	+ NUMBITS_RENDER_TARGET_FORMAT)
+#define OFFSET_STENCIL_ENABLED			(OFFSET_DEPTH_ENABLED			+ NUMBITS_DEPTH_ENABLED)
+#define OFFSET_SAMPLE_COUNT				(OFFSET_STENCIL_ENABLED			+ NUMBITS_STENCIL_ENABLED)
+#define OFFSET_END						(OFFSET_SAMPLE_COUNT			+ NUMBITS_SAMPLE_COUNT)
 
-static uint32 BlendBitOffsets[] = { OFFSET_BLEND_STATE0, OFFSET_BLEND_STATE1, OFFSET_BLEND_STATE2, OFFSET_BLEND_STATE3, };
-static uint32 RTBitOffsets[] = { OFFSET_RENDER_TARGET_FORMAT0, OFFSET_RENDER_TARGET_FORMAT1, OFFSET_RENDER_TARGET_FORMAT2, OFFSET_RENDER_TARGET_FORMAT3, };
+static_assert(OFFSET_END <= 64, "Too many bits used for the Hash");
+
+static uint32 BlendBitOffsets[] = { OFFSET_BLEND_STATE0, OFFSET_BLEND_STATE1, OFFSET_BLEND_STATE2, OFFSET_BLEND_STATE3, OFFSET_BLEND_STATE4, OFFSET_BLEND_STATE5, };
+static uint32 RTBitOffsets[] = { OFFSET_RENDER_TARGET_FORMAT0, OFFSET_RENDER_TARGET_FORMAT1, OFFSET_RENDER_TARGET_FORMAT2, OFFSET_RENDER_TARGET_FORMAT3, OFFSET_RENDER_TARGET_FORMAT4, OFFSET_RENDER_TARGET_FORMAT5, };
 
 #define SET_HASH(Offset, NumBits, Value) \
 	{ \
@@ -96,10 +105,8 @@ id<MTLRenderPipelineState> FPipelineShadow::CreatePipelineStateForBoundShaderSta
 	}
 
 	// depth setting if it's actually used
-//	if (DepthTargetFormat != MTLPixelFormatInvalid)
-	{
-		Desc.depthAttachmentPixelFormat = DepthTargetFormat;
-	}
+	Desc.depthAttachmentPixelFormat = DepthTargetFormat;
+	Desc.stencilAttachmentPixelFormat = StencilTargetFormat;
 
 	// set the bound shader state settings
 	Desc.vertexDescriptor = BSS->VertexDeclaration->Layout;
@@ -113,14 +120,17 @@ id<MTLRenderPipelineState> FPipelineShadow::CreatePipelineStateForBoundShaderSta
 	id<MTLRenderPipelineState> PipelineState = [FMetalManager::GetDevice() newRenderPipelineStateWithDescriptor:Desc error : &Error];
 	TRACK_OBJECT(PipelineState);
 
-	[Desc release];
-
 	if (PipelineState == nil)
 	{
 		NSLog(@"Failed to generate a pipeline state object: %@", Error);
+		NSLog(@"Vertex shader: %@", BSS->VertexShader->GlslCodeNSString);
+		NSLog(@"Pixel shader: %@", BSS->PixelShader->GlslCodeNSString);
+		NSLog(@"Descriptor: %@", Desc);
 		return nil;
 	}
 
+	[Desc release];
+
 	return PipelineState;
 }
 
@@ -155,12 +165,13 @@ FMetalManager::FMetalManager()
 	, PreviousNumRenderTargets(0)
 	, CurrentDepthRenderTexture(nil)
 	, PreviousDepthRenderTexture(nil)
+	, CurrentStencilRenderTexture(nil)
+	, PreviousStencilRenderTexture(nil)
 	, CurrentMSAARenderTexture(nil)
 	, RingBuffer(Device, RingBufferSize, BufferOffsetAlignment)
 	, QueryBuffer(Device, 64 * 1024, 8)
 	, WhichFreeList(0)
-	, CommandBufferIndex(0)
-	, CompletedCommandBufferIndex(0)
+	, CurrentQueryEventIndex(-1)
 	, SceneFrameCounter(0)
 	, ResourceTableFrameCounter(INDEX_NONE)
 {
@@ -173,6 +184,8 @@ FMetalManager::FMetalManager()
 	FMemory::MemSet(PreviousRenderTargetsViewInfo, 0);
 	FMemory::MemSet(CurrentDepthViewInfo, 0);
 	FMemory::MemSet(PreviousDepthViewInfo, 0);
+	FMemory::MemSet(CurrentStencilViewInfo, 0);
+	FMemory::MemSet(PreviousStencilViewInfo, 0);
 
 	CommandQueue = [Device newCommandQueue];
 
@@ -246,6 +259,9 @@ void FMetalManager::BeginFrame()
 
 void FMetalManager::InitFrame()
 {
+	// reset the event index for this frame
+	CurrentQueryEventIndex = -1;
+
 	// start an auto release pool (EndFrame will drain and remake)
 	CreateAutoreleasePool();
 
@@ -259,6 +275,9 @@ void FMetalManager::InitFrame()
 	// mark us to get later
 	BackBuffer->Surface.Texture = nil;
 
+	// make sure first SetRenderTarget goes through
+	PreviousRenderTargetsInfo.NumColorRenderTargets = -1;
+
 //	double End = FPlatformTime::Seconds();
 //	NSLog(@"Semaphore Block Time: %.2f   -- MakeDrawable Time: %.2f", 1000.0f*(Mid - Start), 1000.0f*(End - Mid));
 	
@@ -280,10 +299,28 @@ void FMetalManager::CreateCurrentCommandBuffer(bool bWait)
 	[CurrentCommandBuffer retain];
 	TRACK_OBJECT(CurrentCommandBuffer);
 
-	uint64 LocalCommandBufferIndex = CommandBufferIndex++;
-	[CurrentCommandBuffer addScheduledHandler : ^ (id <MTLCommandBuffer> Buffer)
+	// move to the next event index
+	CurrentQueryEventIndex++;
+
+	FEvent* LocalEvent;
+	// make a new event if we don't have enough
+	if (CurrentQueryEventIndex >= QueryEvents[WhichFreeList].Num())
 	{
-		FMetalManager::Get()->SetCompletedCommandBufferIndex(LocalCommandBufferIndex);
+		LocalEvent = FPlatformProcess::CreateSynchEvent(true);
+		CurrentQueryEventIndex = QueryEvents[WhichFreeList].Add(LocalEvent);
+	}
+	// otherwise, reuse an old one
+	else
+	{
+		LocalEvent = QueryEvents[WhichFreeList][CurrentQueryEventIndex];
+	}
+
+	// unset the event
+	LocalEvent->Reset();
+	[CurrentCommandBuffer addCompletedHandler:^(id <MTLCommandBuffer> Buffer)
+	{
+		// when the command buffer is done, release the event
+		LocalEvent->Trigger();
 	}];
 }
 
@@ -435,19 +472,15 @@ void FMetalManager::SetBoundShaderState(FMetalBoundShaderState* BoundShaderState
 	CurrentBoundShaderState = BoundShaderState;
 }
 
-void FMetalManager::SetCurrentRenderTarget(FMetalSurface* RenderSurface, int32 RenderTargetIndex, uint32 MipIndex, uint32 ArraySliceIndex, MTLLoadAction LoadAction, MTLStoreAction StoreAction, int32 TotalNumRenderTargets)
-{
-	// rememeber our new max
-	CurrentNumRenderTargets = TotalNumRenderTargets;
-	
-	// user code generally passes -1 as a default, but we need 0
-	ArraySliceIndex = ArraySliceIndex == 0xFFFFFFFF ? 0 : ArraySliceIndex;
 
-	// update the current rendered-to pixel format
-	if (RenderSurface)
+
+void FMetalManager::ConditionalUpdateBackBuffer(FMetalSurface& Surface)
+{
+	// are we setting the back buffer? if so, make sure we have the drawable
+	if (&Surface == &BackBuffer->Surface)
 	{
-		// first time in a frame that we are setting the backbuffer, get it
-		if (RenderSurface == &BackBuffer->Surface && RenderSurface->Texture == nil && CurrentDrawable == nil)
+		// update the back buffer texture the first time used this frame
+		if (Surface.Texture == nil && CurrentDrawable == nil)
 		{
 			SCOPE_CYCLE_COUNTER(STAT_MetalMakeDrawableTime);
 
@@ -460,192 +493,191 @@ void FMetalManager::SetCurrentRenderTarget(FMetalSurface* RenderSurface, int32 R
 			GRenderThreadNumIdle[ERenderThreadIdleTypes::WaitingForGPUPresent]++;
 
 			// set the texture into the backbuffer
-			RenderSurface->Texture = CurrentDrawable.texture;
+			Surface.Texture = CurrentDrawable.texture;
 		}
-		
-		if (RenderSurface->bIsCubemap)
-		{
-			ArraySliceIndex = GetMetalCubeFace((ECubeFace)ArraySliceIndex);
-		}
-
-		CurrentColorRenderTextures[RenderTargetIndex] = RenderSurface->Texture;
-		CurrentRenderTargetsViewInfo[RenderTargetIndex].MipIndex = MipIndex;
-		CurrentRenderTargetsViewInfo[RenderTargetIndex].ArraySliceIndex = ArraySliceIndex;
-		CurrentRenderTargetsViewInfo[RenderTargetIndex].LoadAction = LoadAction;
-		CurrentRenderTargetsViewInfo[RenderTargetIndex].StoreAction = StoreAction;
-
-		// only allow one MRT when using MSAA
-		checkf(RenderSurface->MSAATexture == NULL || TotalNumRenderTargets == 1);
-		CurrentMSAARenderTexture = RenderSurface->MSAATexture;
-	}
-	else
-	{
-		CurrentColorRenderTextures[RenderTargetIndex] = nil;
-		CurrentRenderTargetsViewInfo[RenderTargetIndex].MipIndex = 0;
-		CurrentRenderTargetsViewInfo[RenderTargetIndex].ArraySliceIndex = 0;
-		CurrentRenderTargetsViewInfo[RenderTargetIndex].LoadAction = MTLLoadActionDontCare;
-		CurrentRenderTargetsViewInfo[RenderTargetIndex].StoreAction = MTLStoreActionStore;
 	}
 }
 
-void FMetalManager::SetCurrentDepthStencilTarget(FMetalSurface* RenderSurface, MTLLoadAction LoadAction, MTLStoreAction StoreAction, float ClearDepthValue)
+bool FMetalManager::NeedsToSetRenderTarget(const FRHISetRenderTargetsInfo& RenderTargetsInfo)
 {
-	if (RenderSurface)
-	{
-		// @todo metal stencil: track stencil here
-		CurrentDepthRenderTexture = RenderSurface->Texture;
-		CurrentDepthViewInfo.LoadAction = LoadAction;
-		CurrentDepthViewInfo.StoreAction = StoreAction;
-		CurrentDepthViewInfo.ClearDepthValue = ClearDepthValue;
-	}
-	else
-	{
-		CurrentDepthRenderTexture = nil;
-		CurrentDepthViewInfo.LoadAction = MTLLoadActionClear;
-		CurrentDepthViewInfo.StoreAction = MTLStoreActionDontCare;
-		CurrentDepthViewInfo.ClearDepthValue = 0.0f;
-	}
-}
+	// see if our new Info matches our previous Info
 
-void FMetalManager::UpdateContext()
-{
-	// if all render targets match, we can early out
-	if (CurrentNumRenderTargets == PreviousNumRenderTargets && CurrentDepthRenderTexture == PreviousDepthRenderTexture)
+	// basic checks
+	bool bAllChecksPassed = RenderTargetsInfo.NumColorRenderTargets == PreviousRenderTargetsInfo.NumColorRenderTargets &&
+		// handle the case where going from backbuffer + depth -> backbuffer + null, no need to reset RT and do a store/load
+		(RenderTargetsInfo.DepthStencilRenderTarget.Texture == PreviousRenderTargetsInfo.DepthStencilRenderTarget.Texture || RenderTargetsInfo.DepthStencilRenderTarget.Texture == nullptr);
+
+	// now check each color target if the basic tests passe
+	if (bAllChecksPassed)
 	{
-		// make sure all match
-		bool bAllMatch = true;
-		for (uint32 AttachmentIndex = 0; AttachmentIndex < CurrentNumRenderTargets; AttachmentIndex++)
+		for (int32 RenderTargetIndex = 0; RenderTargetIndex < RenderTargetsInfo.NumColorRenderTargets; RenderTargetIndex++)
 		{
-			if (CurrentColorRenderTextures[AttachmentIndex] != PreviousColorRenderTextures[AttachmentIndex] ||
-				CurrentRenderTargetsViewInfo[AttachmentIndex].MipIndex != PreviousRenderTargetsViewInfo[AttachmentIndex].MipIndex ||
-				CurrentRenderTargetsViewInfo[AttachmentIndex].ArraySliceIndex != PreviousRenderTargetsViewInfo[AttachmentIndex].ArraySliceIndex)
+			const FRHIRenderTargetView& RenderTargetView = RenderTargetsInfo.ColorRenderTarget[RenderTargetIndex];
+			const FRHIRenderTargetView& PreviousRenderTargetView = PreviousRenderTargetsInfo.ColorRenderTarget[RenderTargetIndex];
+
+			if (!(RenderTargetView == PreviousRenderTargetView))
 			{
-				bAllMatch = false;
+				bAllChecksPassed = false;
+				break;
 			}
 		}
-
-		if (bAllMatch)
-		{
-			return;
-		}
-
-		//@todo-rco: Do we need to test changes in Load/Store actions (and/or clear values) for Color & Depth?
 	}
 
-	// handle the case where going from backbuffer + depth -> backbuffer + null, no need to reset RT and do a store/load
-	if (CurrentNumRenderTargets == 1 && CurrentColorRenderTextures[0] == PreviousColorRenderTextures[0] && CurrentDepthRenderTexture == nil)
+	// if we are setting them to nothing, then this is probably end of frame, and we can't make a framebuffer
+	// with nothng, so just abort this (only need to check on single MRT case)
+	if (RenderTargetsInfo.NumColorRenderTargets == 1 && RenderTargetsInfo.ColorRenderTarget[0].Texture == nullptr &&
+		RenderTargetsInfo.DepthStencilRenderTarget.Texture == nullptr)
 	{
-		return;
+		bAllChecksPassed = true;
 	}
 
-	PreviousNumRenderTargets = CurrentNumRenderTargets;
+	return bAllChecksPassed == false;
+}
 
-	// if we are setting them to nothing, then this is probably end of frame, and we can't make a framebuffer
-	// with nothng, so just abort this (only need to check on single MRT case)
-	if (CurrentNumRenderTargets == 1 && CurrentColorRenderTextures[0] == nil && CurrentDepthRenderTexture == nil)
+void FMetalManager::SetRenderTargetsInfo(const FRHISetRenderTargetsInfo& RenderTargetsInfo)
+{
+	// see if our new Info matches our previous Info
+	if (NeedsToSetRenderTarget(RenderTargetsInfo) == false)
 	{
 		return;
 	}
 
-	// make a new one (autoreleased)
-	MTLRenderPassDescriptor* CurrentRenderPass = [MTLRenderPassDescriptor renderPassDescriptor];
+	// back this up for next frame
+	PreviousRenderTargetsInfo = RenderTargetsInfo;
+
+	// at this point, we need to fully set up an encoder/command buffer, so make a new one (autoreleased)
+	MTLRenderPassDescriptor* RenderPass = [MTLRenderPassDescriptor renderPassDescriptor];
 
 	// if we need to do queries, write to the ring buffer (we set the offset into the ring buffer per query)
-	CurrentRenderPass.visibilityResultBuffer = QueryBuffer.Buffer;
+	RenderPass.visibilityResultBuffer = QueryBuffer.Buffer;
 
 	// default to non-msaa
 	Pipeline.SampleCount = 0;
 
-	for (uint32 AttachmentIndex = 0; AttachmentIndex < MaxMetalRenderTargets; AttachmentIndex++)
+	for (uint32 RenderTargetIndex = 0; RenderTargetIndex < MaxMetalRenderTargets; RenderTargetIndex++)
 	{
 		// only try to set it if it was one that was set (ie less than CurrentNumRenderTargets)
-		if (AttachmentIndex < CurrentNumRenderTargets && CurrentColorRenderTextures[AttachmentIndex])
+		if (RenderTargetIndex < RenderTargetsInfo.NumColorRenderTargets && RenderTargetsInfo.ColorRenderTarget[RenderTargetIndex].Texture != nullptr)
 		{
+			const FRHIRenderTargetView& RenderTargetView = RenderTargetsInfo.ColorRenderTarget[RenderTargetIndex];
+			FMetalSurface& Surface = GetMetalSurfaceFromRHITexture(RenderTargetView.Texture);
+		
+			// if this is the back buffer, make sure we have a usable drawable
+			ConditionalUpdateBackBuffer(Surface);
+
+			// user code generally passes -1 as a default, but we need 0
+			uint32 ArraySliceIndex = RenderTargetView.ArraySliceIndex == 0xFFFFFFFF ? 0 : RenderTargetView.ArraySliceIndex;
+			if (Surface.bIsCubemap)
+			{
+				ArraySliceIndex = GetMetalCubeFace((ECubeFace)ArraySliceIndex);
+			}
+
 			MTLRenderPassColorAttachmentDescriptor* ColorAttachment = [[MTLRenderPassColorAttachmentDescriptor alloc] init];
 
-			if (CurrentMSAARenderTexture)
+			if (Surface.MSAATexture != nil)
 			{
 				// set up an MSAA attachment
-				ColorAttachment.texture = CurrentMSAARenderTexture;
-				[ColorAttachment setStoreAction:MTLStoreActionMultisampleResolve];
-				[ColorAttachment setResolveTexture:CurrentColorRenderTextures[AttachmentIndex]];
-				Pipeline.SampleCount = CurrentMSAARenderTexture.sampleCount;
+				ColorAttachment.texture = Surface.MSAATexture;
+				ColorAttachment.storeAction = MTLStoreActionMultisampleResolve;
+				ColorAttachment.resolveTexture = Surface.Texture;
+				Pipeline.SampleCount = Surface.MSAATexture.sampleCount;
 
 				// only allow one MRT with msaa
-				checkf(CurrentNumRenderTargets == 1, TEXT("Only expected one MRT when using MSAA"));
+				checkf(RenderTargetsInfo.NumColorRenderTargets == 1, TEXT("Only expected one MRT when using MSAA"));
 			}
 			else
 			{
 				// set up non-MSAA attachment
-				ColorAttachment.texture = CurrentColorRenderTextures[AttachmentIndex];
-				[ColorAttachment setStoreAction:MTLStoreActionStore];
+				ColorAttachment.texture = Surface.Texture;
+				ColorAttachment.storeAction = GetMetalRTStoreAction(RenderTargetView.StoreAction);
 				Pipeline.SampleCount = 1;
 			}
-			ColorAttachment.level = CurrentRenderTargetsViewInfo[AttachmentIndex].MipIndex;
-			ColorAttachment.slice = CurrentRenderTargetsViewInfo[AttachmentIndex].ArraySliceIndex;
-			ColorAttachment.LoadAction = CurrentRenderTargetsViewInfo[AttachmentIndex].LoadAction;
-			//@todo implement store, but making sure that multisampleresolve is handled properly.
+			ColorAttachment.level = RenderTargetView.MipIndex;
+			ColorAttachment.slice = ArraySliceIndex;
+			ColorAttachment.loadAction = GetMetalRTLoadAction(RenderTargetView.LoadAction);
+			const FLinearColor& ClearColor = RenderTargetsInfo.ClearColors[RenderTargetIndex];
+			ColorAttachment.clearColor = MTLClearColorMake(ClearColor.R, ClearColor.G, ClearColor.B, ClearColor.A);
 
 			// assign the attachment to the slot
-			[CurrentRenderPass.colorAttachments setObject:ColorAttachment atIndexedSubscript:AttachmentIndex];
-			[ColorAttachment release];
+			[RenderPass.colorAttachments setObject:ColorAttachment atIndexedSubscript:RenderTargetIndex];
+
+			Pipeline.RenderTargets[RenderTargetIndex].pixelFormat = ColorAttachment.texture.pixelFormat;
 
-			Pipeline.RenderTargets[AttachmentIndex].pixelFormat = CurrentColorRenderTextures[AttachmentIndex].pixelFormat;
+			[ColorAttachment release];
 		}
 		else
 		{
-			Pipeline.RenderTargets[AttachmentIndex].pixelFormat = MTLPixelFormatInvalid;
+			Pipeline.RenderTargets[RenderTargetIndex].pixelFormat = MTLPixelFormatInvalid;
 		}
-		
-		// update the hash no matter what case (null, unused, used)
-		SET_HASH(RTBitOffsets[AttachmentIndex], NUMBITS_RENDER_TARGET_FORMAT, Pipeline.RenderTargets[AttachmentIndex].pixelFormat);
 
-		// remember this for next time
-		PreviousColorRenderTextures[AttachmentIndex] = CurrentColorRenderTextures[AttachmentIndex];
-		PreviousRenderTargetsViewInfo[AttachmentIndex] = CurrentRenderTargetsViewInfo[AttachmentIndex];
+		// update the hash no matter what case (null, unused, used)
+		SET_HASH(RTBitOffsets[RenderTargetIndex], NUMBITS_RENDER_TARGET_FORMAT, Pipeline.RenderTargetFormatKeys[RenderTargetIndex]);
 	}
-	
-	if (CurrentDepthRenderTexture)
-	{
-		MTLRenderPassDepthAttachmentDescriptor* DepthAttachment = [[MTLRenderPassDepthAttachmentDescriptor alloc] init];
 
-		// set up the depth attachment
-		DepthAttachment.texture = CurrentDepthRenderTexture;
-		[DepthAttachment setLoadAction:CurrentDepthViewInfo.LoadAction];
-		[DepthAttachment setStoreAction : CurrentDepthViewInfo.StoreAction];
-		[DepthAttachment setClearDepth : CurrentDepthViewInfo.ClearDepthValue];
+	// default to invalid
+	Pipeline.DepthTargetFormat = MTLPixelFormatInvalid;
+	Pipeline.StencilTargetFormat = MTLPixelFormatInvalid;
 
-		Pipeline.DepthTargetFormat = CurrentDepthRenderTexture.pixelFormat;
-		if (Pipeline.SampleCount == 0)
+	// setup depth and/or stencil
+	if (RenderTargetsInfo.DepthStencilRenderTarget.Texture != nullptr)
+	{
+		FMetalSurface& Surface= GetMetalSurfaceFromRHITexture(RenderTargetsInfo.DepthStencilRenderTarget.Texture);
+		if (Surface.Texture != nil)
 		{
-			Pipeline.SampleCount = CurrentDepthRenderTexture.sampleCount;
+			MTLRenderPassDepthAttachmentDescriptor* DepthAttachment = [[MTLRenderPassDepthAttachmentDescriptor alloc] init];
+
+			// set up the depth attachment
+			DepthAttachment.texture = Surface.Texture;
+			DepthAttachment.loadAction = GetMetalRTLoadAction(RenderTargetsInfo.DepthStencilRenderTarget.DepthLoadAction);
+			DepthAttachment.storeAction = GetMetalRTStoreAction(RenderTargetsInfo.DepthStencilRenderTarget.DepthStoreAction);
+			DepthAttachment.clearDepth = RenderTargetsInfo.DepthClearValue;
+
+			Pipeline.DepthTargetFormat = DepthAttachment.texture.pixelFormat;
+			if (Pipeline.SampleCount == 0)
+			{
+				Pipeline.SampleCount = DepthAttachment.texture.sampleCount;
+			}
+
+			// and assign it
+			RenderPass.depthAttachment = DepthAttachment;
+			[DepthAttachment release];
 		}
 
-		// and assign it
-		CurrentRenderPass.depthAttachment = DepthAttachment;
-		[DepthAttachment release];
-	}
-	else
-	{
-		Pipeline.DepthTargetFormat = MTLPixelFormatInvalid;
-	}
-	// update hash for the depth buffer
-	SET_HASH(OFFSET_DEPTH_TARGET_FORMAT, NUMBITS_DEPTH_TARGET_FORMAT, Pipeline.DepthTargetFormat);
+		if (Surface.StencilTexture != nil)
+		{
+			MTLRenderPassStencilAttachmentDescriptor* StencilAttachment = [[MTLRenderPassStencilAttachmentDescriptor alloc] init];
 
-	// remember this for next time
-	PreviousDepthViewInfo = CurrentDepthViewInfo;
-	PreviousDepthRenderTexture = CurrentDepthRenderTexture;
+			// set up the stencil attachment
+			StencilAttachment.texture = Surface.StencilTexture;
+			StencilAttachment.loadAction = GetMetalRTLoadAction(RenderTargetsInfo.DepthStencilRenderTarget.StencilLoadAction);
+			StencilAttachment.storeAction = GetMetalRTStoreAction(RenderTargetsInfo.DepthStencilRenderTarget.StencilStoreAction);
+			StencilAttachment.clearStencil = RenderTargetsInfo.StencilClearValue;
 
+			Pipeline.StencilTargetFormat = StencilAttachment.texture.pixelFormat;
+			if (Pipeline.SampleCount == 0)
+			{
+				Pipeline.SampleCount = StencilAttachment.texture.sampleCount;
+			}
+
+			// and assign it
+			RenderPass.stencilAttachment = StencilAttachment;
+			[StencilAttachment release];
+		}
+	}
+
+	// update hash for the depth buffer
+	SET_HASH(OFFSET_DEPTH_ENABLED, NUMBITS_DEPTH_ENABLED, (Pipeline.DepthTargetFormat == MTLPixelFormatInvalid ? 0 : 1));
+	SET_HASH(OFFSET_STENCIL_ENABLED, NUMBITS_STENCIL_ENABLED, (Pipeline.StencilTargetFormat == MTLPixelFormatInvalid ? 0 : 1));
 	SET_HASH(OFFSET_SAMPLE_COUNT, NUMBITS_SAMPLE_COUNT, Pipeline.SampleCount);
 
 	// commit pending commands on the old render target
 	if (CurrentContext)
 	{
 #if UE_BUILD_DEBUG || UE_BUILD_DEVELOPMENT
-		if (NumDrawCalls == 0)
-		{
-			NSLog(@"There were %d draw calls for an RT in frame %lld", NumDrawCalls, GFrameCounter);
-		}
+//		if (NumDrawCalls == 0)
+//		{
+//			NSLog(@"There were %d draw calls for an RT in frame %lld", NumDrawCalls, GFrameCounter);
+//		}
 #endif
 
 		[CurrentContext endEncoding];
@@ -653,28 +685,24 @@ void FMetalManager::UpdateContext()
 
 		[CurrentContext release];
 
-		// if we are doing occlusion queries, we could use this method, along with a completion callback
-		// to set a "render target complete" flag that the OQ code next frame would wait on
 		// commit the buffer for this context
 		[CurrentCommandBuffer commit];
 		UNTRACK_OBJECT(CurrentCommandBuffer);
 		[CurrentCommandBuffer release];
-		
+
 		// create the command buffer for this frame
 		CreateCurrentCommandBuffer(false);
 	}
-		
+
 	// make a new render context to use to render to the framebuffer
-	CurrentContext = [CurrentCommandBuffer renderCommandEncoderWithDescriptor:CurrentRenderPass];
+	CurrentContext = [CurrentCommandBuffer renderCommandEncoderWithDescriptor:RenderPass];
 	[CurrentContext retain];
 	TRACK_OBJECT(CurrentContext);
 
-	// make suire the rasterizer state is set the first time for each new encoder
+	// make sure the rasterizer state is set the first time for each new encoder
 	bFirstRasterizerState = true;
 }
 
-
-
 uint32 FRingBuffer::Allocate(uint32 Size, uint32 Alignment)
 {
 	if (Alignment == 0)
@@ -823,45 +851,6 @@ void FMetalManager::CommitNonComputeShaderConstants()
 	}
 }
 
-bool FMetalManager::WaitForCommandBufferComplete(uint64 IndexToWaitFor, double Timeout)
-{
-	// don't track a block if not needed
-	if (CompletedCommandBufferIndex >= IndexToWaitFor)
-	{
-		return true;
-	}
-
-	// if we don't want to wait, then we have failed
-	if (Timeout == 0.0)
-	{
-		return false;
-	}
-
-	// if we block until it's ready, loop here until it is
-	SCOPE_CYCLE_COUNTER(STAT_RenderQueryResultTime);
-	uint32 IdleStart = FPlatformTime::Cycles();
-	double StartTime = FPlatformTime::Seconds();
-
-	while (CompletedCommandBufferIndex < IndexToWaitFor)
-	{
-		FPlatformProcess::Sleep(0);
-
-		// look for gpu stuck/crashed
-		if ((FPlatformTime::Seconds() - StartTime) > Timeout)
-		{
-			UE_LOG(LogRHI, Log, TEXT("Timed out while waiting for GPU to catch up on occlusion/timer results. (%.1f s)"), Timeout);
-			return false;
-		}
-	}
-
-	// track idle time blocking on GPU
-	GRenderThreadIdle[ERenderThreadIdleTypes::WaitingForGPUQuery] += FPlatformTime::Cycles() - IdleStart;
-	GRenderThreadNumIdle[ERenderThreadIdleTypes::WaitingForGPUQuery]++;
-
-	return true;
-}
-
-
 void FMetalManager::SetRasterizerState(const FRasterizerStateInitializerRHI& State)
 {
 	if (bFirstRasterizerState)
@@ -893,4 +882,4 @@ void FMetalManager::SetRasterizerState(const FRasterizerStateInitializerRHI& Sta
 #endif
 
 	bFirstRasterizerState = false;
-}
+}
\ No newline at end of file
diff --git a/Engine/Source/Runtime/IOS/MetalRHI/Private/MetalManager.h b/Engine/Source/Runtime/IOS/MetalRHI/Private/MetalManager.h
--- a/Engine/Source/Runtime/IOS/MetalRHI/Private/MetalManager.h
+++ b/Engine/Source/Runtime/IOS/MetalRHI/Private/MetalManager.h
@@ -40,7 +40,9 @@ struct FPipelineShadow
 	id<MTLRenderPipelineState> CreatePipelineStateForBoundShaderState(FMetalBoundShaderState* BSS) const;
 
 	MTLRenderPipelineColorAttachmentDescriptor* RenderTargets[MaxMetalRenderTargets];
+	uint8 RenderTargetFormatKeys[MaxMetalRenderTargets];
 	MTLPixelFormat DepthTargetFormat;
+	MTLPixelFormat StencilTargetFormat;
     uint32 SampleCount;
 
 	// running hash of the pipeline state
@@ -100,8 +102,14 @@ public:
 	void SetBlendState(FMetalBlendState* BlendState);
 	void SetBoundShaderState(FMetalBoundShaderState* BoundShaderState);
 	void SetCurrentRenderTarget(FMetalSurface* RenderSurface, int32 RenderTargetIndex, uint32 MipIndex, uint32 ArraySliceIndex, MTLLoadAction LoadAction, MTLStoreAction StoreAction, int32 TotalNumRenderTargets);
-	void SetCurrentDepthStencilTarget(FMetalSurface* RenderSurface, MTLLoadAction LoadAction, MTLStoreAction StoreAction, float ClearDepthValue);
+	void SetCurrentDepthStencilTarget(FMetalSurface* RenderSurface, MTLLoadAction DepthLoadAction=MTLLoadActionDontCare, MTLStoreAction DepthStoreAction=MTLStoreActionDontCare, float ClearDepthValue=0, 
+		MTLLoadAction StencilLoadAction=MTLLoadActionDontCare, MTLStoreAction StencilStoreAction=MTLStoreActionDontCare, uint8 ClearStencilValue=0);
 	
+	/**
+	 * Set the color, depth and stencil render targets, and then make the new command buffer/encoder
+	 */
+	void SetRenderTargetsInfo(const FRHISetRenderTargetsInfo& RenderTargetsInfo);
+
 	/**
 	 * Update the context with the current render targets
 	 */
@@ -122,18 +130,13 @@ public:
 		return QueryBuffer.Buffer;
 	}
 
-	uint64 GetCommandBufferIndex()
+	FEvent* GetCurrentQueryCompleteEvent()
 	{
-		return CommandBufferIndex;
-	}
-	
-	void SetCompletedCommandBufferIndex(uint64 Index)
-	{
-		CompletedCommandBufferIndex = Index;
+		checkf(CurrentQueryEventIndex < QueryEvents[WhichFreeList].Num(), TEXT("Attempted to get the current query event, but it doesn't exist"));
+		return QueryEvents[WhichFreeList][CurrentQueryEventIndex];
 	}
 
 	void SubmitCommandBufferAndWait();
-	bool WaitForCommandBufferComplete(uint64 IndexToWaitFor, double Timeout);
 
 	void SetRasterizerState(const FRasterizerStateInitializerRHI& State);
 
@@ -165,6 +168,16 @@ protected:
 
 	void CreateCurrentCommandBuffer(bool bWait);
 
+	/**
+	 * If the given surface is rendering to the back buffer, make sure it's fixed up
+	 */
+	void ConditionalUpdateBackBuffer(FMetalSurface& Surface);
+
+	/**
+	 * Check to see if the new RenderTargetInfo needs to create an entire new encoder/command buffer
+	 */
+	bool NeedsToSetRenderTarget(const FRHISetRenderTargetsInfo& RenderTargetsInfo);
+
 
 	id<MTLDevice> Device;
 
@@ -193,11 +206,22 @@ protected:
 		MTLStoreAction StoreAction;
 	};
 
+	struct FStencilRenderTargetViewInfo
+	{
+		uint8 ClearStencilValue;
+		MTLLoadAction LoadAction;
+		MTLStoreAction StoreAction;
+	};
+
+	FRHISetRenderTargetsInfo PreviousRenderTargetsInfo;
+
 	FRenderTargetViewInfo CurrentRenderTargetsViewInfo[MaxMetalRenderTargets], PreviousRenderTargetsViewInfo[MaxMetalRenderTargets];
 	uint32 CurrentNumRenderTargets, PreviousNumRenderTargets;
 	id<MTLTexture> CurrentColorRenderTextures[MaxMetalRenderTargets], PreviousColorRenderTextures[MaxMetalRenderTargets];
 	FDepthRenderTargetViewInfo CurrentDepthViewInfo, PreviousDepthViewInfo;
+	FStencilRenderTargetViewInfo CurrentStencilViewInfo, PreviousStencilViewInfo;
 	id<MTLTexture> CurrentDepthRenderTexture, PreviousDepthRenderTexture;
+	id<MTLTexture> CurrentStencilRenderTexture, PreviousStencilRenderTexture;
 	id<MTLTexture> CurrentMSAARenderTexture;
 	
 	TRefCountPtr<FMetalTexture2D> BackBuffer;
@@ -218,6 +242,8 @@ protected:
 
 	TArray<id> DelayedFreeLists[4];
 	uint32 WhichFreeList;
+	TArray<FEvent*> QueryEvents[4];
+	int32 CurrentQueryEventIndex;
 
 	// the slot to store a per-thread autorelease pool
 	uint32 AutoReleasePoolTLSSlot;
diff --git a/Engine/Source/Runtime/IOS/MetalRHI/Private/MetalQuery.cpp b/Engine/Source/Runtime/IOS/MetalRHI/Private/MetalQuery.cpp
--- a/Engine/Source/Runtime/IOS/MetalRHI/Private/MetalQuery.cpp
+++ b/Engine/Source/Runtime/IOS/MetalRHI/Private/MetalQuery.cpp
@@ -28,7 +28,7 @@ void FMetalRenderQuery::Begin()
 	*QueryMemory = 0;
 
 	// remember which command buffer we are in
-	CommandBufferIndex = FMetalManager::Get()->GetCommandBufferIndex();
+	QueryCompleteEvent = FMetalManager::Get()->GetCurrentQueryCompleteEvent();
 
 	[FMetalManager::GetContext() setVisibilityResultMode:MTLVisibilityResultModeBoolean offset:Offset];
 }
@@ -48,28 +48,30 @@ FRenderQueryRHIRef FMetalDynamicRHI::RHICreateRenderQuery(ERenderQueryType Query
 
 void FMetalDynamicRHI::RHIResetRenderQuery(FRenderQueryRHIParamRef QueryRHI)
 {
-	DYNAMIC_CAST_METGALRESOURCE(RenderQuery,Query);
+	DYNAMIC_CAST_METALRESOURCE(RenderQuery,Query);
 
 //	NOT_SUPPORTED("RHIResetRenderQuery");
 }
 
 bool FMetalDynamicRHI::RHIGetRenderQueryResult(FRenderQueryRHIParamRef QueryRHI,uint64& OutNumPixels,bool bWait)
 {
-	DYNAMIC_CAST_METGALRESOURCE(RenderQuery, Query);
+	DYNAMIC_CAST_METALRESOURCE(RenderQuery, Query);
 
 	uint32* QueryMemory = (uint32*)((uint8*)[FMetalManager::Get()->GetQueryBuffer() contents] + Query->Offset);
 
-	// timer queries are used for Benchmarks which can stall a bit more
-	double Timeout = 0.0;
 	if (bWait)
 	{
-		Timeout = 0.5f;//(Query->QueryType == RQT_AbsoluteTime) ? 2.0 : 0.5;
-	}
-	if (!FMetalManager::Get()->WaitForCommandBufferComplete(Query->CommandBufferIndex, Timeout))
-	{
-		return false;
+		// wait for half a second
+		uint64 WaitTimeMS = 500;
+		Query->QueryCompleteEvent->Wait(WaitTimeMS);
 	}
 
+	// waiting for 0 time on a manual reset event will return true if it was already triggered
+ 	if (Query->QueryCompleteEvent->Wait(0) == false)
+ 	{
+ 		return false;
+ 	}
+
 	// at this point, we are ready to read the value!
 	OutNumPixels = *QueryMemory;
 	return true;
diff --git a/Engine/Source/Runtime/IOS/MetalRHI/Private/MetalRHI.cpp b/Engine/Source/Runtime/IOS/MetalRHI/Private/MetalRHI.cpp
--- a/Engine/Source/Runtime/IOS/MetalRHI/Private/MetalRHI.cpp
+++ b/Engine/Source/Runtime/IOS/MetalRHI/Private/MetalRHI.cpp
@@ -41,9 +41,30 @@ FMetalDynamicRHI::FMetalDynamicRHI()
 	check( IsInGameThread() );
 	check( !GIsThreadedRendering );
 
-	// Initialize the RHI capabilities.
-	GMaxRHIFeatureLevel = ERHIFeatureLevel::ES2;
-	GMaxRHIShaderPlatform = SP_METAL;
+
+	// get the device to ask about capabilities
+	id<MTLDevice> Device = [IOSAppDelegate GetDelegate].IOSView->MetalDevice;
+	// A8 can use 256 bits of MRTs
+	bool bCanUseWideMRTs = [Device supportsFeatureSet:MTLFeatureSet_iOS_GPUFamily2_v1];
+	// only allow GBuffers, etc on A8s (A7s are just not going to cut it)
+	if (bCanUseWideMRTs && FParse::Param(FCommandLine::Get(),TEXT("metalmrt")))
+	{
+		GMaxRHIFeatureLevel = ERHIFeatureLevel::SM4;
+		GMaxRHIShaderPlatform = SP_METAL_MRT;
+	}
+	else
+	{
+		GMaxRHIFeatureLevel = ERHIFeatureLevel::ES2;
+		GMaxRHIShaderPlatform = SP_METAL;
+	}
+
+	if (FPlatformMisc::IsDebuggerPresent() && UE_BUILD_DEBUG)
+	{
+		// Enable GL debug markers if we're running in Xcode
+		extern int32 GEmitMeshDrawEvent;
+		GEmitMeshDrawEvent = 1;
+		GEmitDrawEvents = true;
+	}
 
 	GPixelCenterOffset = 0.0f;
 	GSupportsVertexInstancing = false;
@@ -71,9 +92,12 @@ FMetalDynamicRHI::FMetalDynamicRHI()
 
 	GShaderPlatformForFeatureLevel[ERHIFeatureLevel::ES2] = SP_METAL;
 	GShaderPlatformForFeatureLevel[ERHIFeatureLevel::ES3_1] = SP_NumPlatforms;
-	GShaderPlatformForFeatureLevel[ERHIFeatureLevel::SM4] = SP_NumPlatforms;
+	GShaderPlatformForFeatureLevel[ERHIFeatureLevel::SM4] = (GMaxRHIShaderPlatform == SP_METAL_MRT) ? SP_METAL_MRT : SP_NumPlatforms;
 	GShaderPlatformForFeatureLevel[ERHIFeatureLevel::SM5] = SP_NumPlatforms;
 
+	// @todo metal mrt:
+	GSupportsVolumeTextureRendering = false;
+
 // 	GDrawUPVertexCheckCount = MAX_uint16;
 
 	// Initialize the platform pixel format map.
@@ -98,9 +122,11 @@ FMetalDynamicRHI::FMetalDynamicRHI()
 	GPixelFormats[PF_DepthStencil		].BlockBytes		= 4;
 	GPixelFormats[PF_X24_G8				].PlatformFormat	= MTLPixelFormatStencil8;
 	GPixelFormats[PF_X24_G8				].BlockBytes		= 1;
-	GPixelFormats[PF_ShadowDepth		].PlatformFormat	= MTLPixelFormatDepth32Float;
+	GPixelFormats[PF_ShadowDepth		].PlatformFormat	= GPixelFormats[PF_DepthStencil].PlatformFormat; // all depth formats must be the same, for the pipeline state hash (see NUMBITS_DEPTH_ENABLED)
 	GPixelFormats[PF_R32_FLOAT			].PlatformFormat	= MTLPixelFormatR32Float;
-	GPixelFormats[PF_G16R16				].PlatformFormat	= MTLPixelFormatRG16Unorm;
+	GPixelFormats[PF_G16R16				].PlatformFormat	= MTLPixelFormatInvalid;//MTLPixelFormatRG16Unorm;
+	// we can't render to this in Metal, so mark it as unsupported (although we could texture from it, but we are only using it for render targets)
+	GPixelFormats[PF_G16R16				].Supported			= false;
 	GPixelFormats[PF_G16R16F			].PlatformFormat	= MTLPixelFormatRG16Float;
 	GPixelFormats[PF_G16R16F_FILTER		].PlatformFormat	= MTLPixelFormatRG16Float;
 	GPixelFormats[PF_G32R32F			].PlatformFormat	= MTLPixelFormatRG32Float;
diff --git a/Engine/Source/Runtime/IOS/MetalRHI/Private/MetalRHIPrivate.h b/Engine/Source/Runtime/IOS/MetalRHI/Private/MetalRHIPrivate.h
--- a/Engine/Source/Runtime/IOS/MetalRHI/Private/MetalRHIPrivate.h
+++ b/Engine/Source/Runtime/IOS/MetalRHI/Private/MetalRHIPrivate.h
@@ -8,8 +8,8 @@
 
 #include "Engine.h"
 
-// UE4 has a Max of 8 RTs, but we max out at 4, no need to waste time processing 8
-const uint32 MaxMetalRenderTargets = 4;
+// UE4 has a Max of 8 RTs, but we can spend less time looping with 6
+const uint32 MaxMetalRenderTargets = 6;
 
 // How many possible vertex streams are allowed
 const uint32 MaxMetalStreams = 16;
diff --git a/Engine/Source/Runtime/IOS/MetalRHI/Private/MetalRenderTarget.cpp b/Engine/Source/Runtime/IOS/MetalRHI/Private/MetalRenderTarget.cpp
--- a/Engine/Source/Runtime/IOS/MetalRHI/Private/MetalRenderTarget.cpp
+++ b/Engine/Source/Runtime/IOS/MetalRHI/Private/MetalRenderTarget.cpp
@@ -34,7 +34,7 @@ void FMetalDynamicRHI::RHIReadSurfaceFloatData(FTextureRHIParamRef TextureRHI, F
 	//kick the current command buffer.
 	FMetalManager::Get()->SubmitCommandBufferAndWait();
 
-	DYNAMIC_CAST_METGALRESOURCE(Texture2D, Texture);
+	DYNAMIC_CAST_METALRESOURCE(Texture2D, Texture);
 	
 	FMetalSurface& Surface = GetMetalSurfaceFromRHITexture(TextureRHI);
 	
diff --git a/Engine/Source/Runtime/IOS/MetalRHI/Private/MetalShaders.cpp b/Engine/Source/Runtime/IOS/MetalRHI/Private/MetalShaders.cpp
--- a/Engine/Source/Runtime/IOS/MetalRHI/Private/MetalShaders.cpp
+++ b/Engine/Source/Runtime/IOS/MetalRHI/Private/MetalShaders.cpp
@@ -225,12 +225,12 @@ FMetalBoundShaderState::FMetalBoundShaderState(
 			FGeometryShaderRHIParamRef InGeometryShaderRHI)
 	:	CacheLink(InVertexDeclarationRHI,InVertexShaderRHI,InPixelShaderRHI,InHullShaderRHI,InDomainShaderRHI,InGeometryShaderRHI,this)
 {
-	DYNAMIC_CAST_METGALRESOURCE(VertexDeclaration,InVertexDeclaration);
-	DYNAMIC_CAST_METGALRESOURCE(VertexShader,InVertexShader);
-	DYNAMIC_CAST_METGALRESOURCE(PixelShader,InPixelShader);
-	DYNAMIC_CAST_METGALRESOURCE(HullShader,InHullShader);
-	DYNAMIC_CAST_METGALRESOURCE(DomainShader,InDomainShader);
-	DYNAMIC_CAST_METGALRESOURCE(GeometryShader,InGeometryShader);
+	DYNAMIC_CAST_METALRESOURCE(VertexDeclaration,InVertexDeclaration);
+	DYNAMIC_CAST_METALRESOURCE(VertexShader,InVertexShader);
+	DYNAMIC_CAST_METALRESOURCE(PixelShader,InPixelShader);
+	DYNAMIC_CAST_METALRESOURCE(HullShader,InHullShader);
+	DYNAMIC_CAST_METALRESOURCE(DomainShader,InDomainShader);
+	DYNAMIC_CAST_METALRESOURCE(GeometryShader,InGeometryShader);
 
 	// cache everything
 	VertexDeclaration = InVertexDeclaration;
diff --git a/Engine/Source/Runtime/IOS/MetalRHI/Private/MetalState.cpp b/Engine/Source/Runtime/IOS/MetalRHI/Private/MetalState.cpp
--- a/Engine/Source/Runtime/IOS/MetalRHI/Private/MetalState.cpp
+++ b/Engine/Source/Runtime/IOS/MetalRHI/Private/MetalState.cpp
@@ -286,10 +286,9 @@ FMetalBlendState::FMetalBlendState(const FBlendStateInitializerRHI& Initializer)
 		uint8* Key = BlendSettingsToUniqueKeyMap.Find(BlendBitMask);
 		if (Key == NULL)
 		{
-			NSLog(@"Making a new key [BlendBitMask: %x, Key: %d]", BlendBitMask, NextKey);
 			Key = &BlendSettingsToUniqueKeyMap.Add(BlendBitMask, NextKey++);
 
-			// only giving 5 bits to the key, we need 4 of them
+			// only giving 5 bits to the key, since we need to pack a bunch of them into 64 bits
 			checkf(NextKey < 32, TEXT("Too many unique blend states to fit into the PipelineStateHash"));
 		}
 		// set the key
diff --git a/Engine/Source/Runtime/IOS/MetalRHI/Private/MetalStructuredBuffer.cpp b/Engine/Source/Runtime/IOS/MetalRHI/Private/MetalStructuredBuffer.cpp
--- a/Engine/Source/Runtime/IOS/MetalRHI/Private/MetalStructuredBuffer.cpp
+++ b/Engine/Source/Runtime/IOS/MetalRHI/Private/MetalStructuredBuffer.cpp
@@ -32,7 +32,7 @@ FStructuredBufferRHIRef FMetalDynamicRHI::RHICreateStructuredBuffer(uint32 Strid
 
 void* FMetalDynamicRHI::RHILockStructuredBuffer(FStructuredBufferRHIParamRef StructuredBufferRHI,uint32 Offset,uint32 Size,EResourceLockMode LockMode)
 {
-	DYNAMIC_CAST_METGALRESOURCE(StructuredBuffer,StructuredBuffer);
+	DYNAMIC_CAST_METALRESOURCE(StructuredBuffer,StructuredBuffer);
 
 	return NULL;
 }
diff --git a/Engine/Source/Runtime/IOS/MetalRHI/Private/MetalTexture.cpp b/Engine/Source/Runtime/IOS/MetalRHI/Private/MetalTexture.cpp
--- a/Engine/Source/Runtime/IOS/MetalRHI/Private/MetalTexture.cpp
+++ b/Engine/Source/Runtime/IOS/MetalRHI/Private/MetalTexture.cpp
@@ -6,6 +6,9 @@
 
 #include "MetalRHIPrivate.h"
 
+TMap<MTLPixelFormat, uint8> FMetalSurface::PixelFormatKeyMap;
+uint8 FMetalSurface::NextKey = 1; // 0 is reserved for MTLPixelFormatInvalid
+
 /** Texture reference class. */
 class FMetalTextureReference : public FRHITextureReference
 {
@@ -72,6 +75,7 @@ static bool IsRenderTarget(uint32 Flags)
 FMetalSurface::FMetalSurface(ERHIResourceType ResourceType, EPixelFormat Format, uint32 InSizeX, uint32 InSizeY, uint32 InSizeZ, uint32 NumSamples, bool bArray, uint32 ArraySize, uint32 NumMips, uint32 InFlags, FResourceBulkDataInterface* BulkData)
 	: PixelFormat(Format)
     , MSAATexture(nil)
+	, StencilTexture(nil)
 	, SizeX(InSizeX)
 	, SizeY(InSizeY)
 	, SizeZ(InSizeZ)
@@ -79,6 +83,20 @@ FMetalSurface::FMetalSurface(ERHIResourceType ResourceType, EPixelFormat Format,
 	, Flags(InFlags)
 	, TotalTextureSize(0)
 {
+	// get a unique key for this surface's format
+	MTLPixelFormat MTLFormat = (MTLPixelFormat)GPixelFormats[Format].PlatformFormat;
+	uint8* Key = PixelFormatKeyMap.Find(MTLFormat);
+	if (Key == NULL)
+	{
+		Key = &PixelFormatKeyMap.Add(MTLFormat, NextKey++);
+
+		// only giving 4 bits to the key
+		checkf(NextKey < 16, TEXT("Too many unique pixel formats to fit into the PipelineStateHash"));
+	}
+	// set the key
+	FormatKey = *Key;
+
+
 	// the special back buffer surface will be updated in FMetalManager::BeginFrame - no need to set the texture here
 	if (Flags & TexCreate_Presentable)
 	{
@@ -86,7 +104,6 @@ FMetalSurface::FMetalSurface(ERHIResourceType ResourceType, EPixelFormat Format,
 	}
 
 	bool bIsRenderTarget = IsRenderTarget(Flags);
-	MTLPixelFormat MTLFormat = (MTLPixelFormat)GPixelFormats[Format].PlatformFormat;
 	MTLTextureDescriptor* Desc;
 	
 	if (ResourceType == RRT_TextureCube)
@@ -96,6 +113,17 @@ FMetalSurface::FMetalSurface(ERHIResourceType ResourceType, EPixelFormat Format,
 															   mipmapped:(NumMips > 1)];
 		bIsCubemap = true;
 	}
+	else if (ResourceType == RRT_Texture3D)
+	{
+		Desc = [MTLTextureDescriptor new];
+		Desc.textureType = MTLTextureType3D;
+		Desc.height = SizeX;
+		Desc.width = SizeY;
+		Desc.depth = SizeZ;
+		Desc.pixelFormat = MTLFormat;
+		Desc.arrayLength = 1;
+		Desc.mipmapLevelCount = 1;
+	}
 	else
 	{
 		Desc = [MTLTextureDescriptor texture2DDescriptorWithPixelFormat:MTLFormat
@@ -173,6 +201,17 @@ FMetalSurface::FMetalSurface(ERHIResourceType ResourceType, EPixelFormat Format,
 		}
 	}
 
+	// create a stencil buffer if needed
+	if (GMaxRHIShaderPlatform == SP_METAL_MRT && Format == PF_DepthStencil)
+	{
+		Desc.textureType = MTLTextureType2D;
+		Desc.pixelFormat = MTLPixelFormatStencil8;
+		StencilTexture = [FMetalManager::GetDevice() newTextureWithDescriptor:Desc];
+
+		// 1 byte per texel
+		TotalTextureSize += SizeX * SizeY;
+	}
+
 	// track memory usage
 	if (bIsRenderTarget)
 	{
@@ -404,7 +443,7 @@ void FMetalDynamicRHI::RHIGenerateMips(FTextureRHIParamRef SourceSurfaceRHI)
 
 FTexture2DRHIRef FMetalDynamicRHI::RHIAsyncReallocateTexture2D(FTexture2DRHIParamRef OldTextureRHI, int32 NewMipCount, int32 NewSizeX, int32 NewSizeY, FThreadSafeCounter* RequestStatus)
 {
-	DYNAMIC_CAST_METGALRESOURCE(Texture2D,OldTexture);
+	DYNAMIC_CAST_METALRESOURCE(Texture2D,OldTexture);
 
 	FMetalTexture2D* NewTexture = new FMetalTexture2D(OldTexture->GetFormat(), NewSizeX, NewSizeY, NewMipCount, OldTexture->GetNumSamples(), OldTexture->GetFlags(), NULL);
 
@@ -467,37 +506,37 @@ ETextureReallocationStatus FMetalDynamicRHI::RHICancelAsyncReallocateTexture2D(
 
 void* FMetalDynamicRHI::RHILockTexture2D(FTexture2DRHIParamRef TextureRHI,uint32 MipIndex,EResourceLockMode LockMode,uint32& DestStride,bool bLockWithinMiptail)
 {
-	DYNAMIC_CAST_METGALRESOURCE(Texture2D,Texture);
+	DYNAMIC_CAST_METALRESOURCE(Texture2D,Texture);
 	return Texture->Surface.Lock(MipIndex, 0, LockMode, DestStride);
 }
 
 void FMetalDynamicRHI::RHIUnlockTexture2D(FTexture2DRHIParamRef TextureRHI,uint32 MipIndex,bool bLockWithinMiptail)
 {
-	DYNAMIC_CAST_METGALRESOURCE(Texture2D,Texture);
+	DYNAMIC_CAST_METALRESOURCE(Texture2D,Texture);
 	Texture->Surface.Unlock(MipIndex, 0);
 }
 
 void* FMetalDynamicRHI::RHILockTexture2DArray(FTexture2DArrayRHIParamRef TextureRHI,uint32 TextureIndex,uint32 MipIndex,EResourceLockMode LockMode,uint32& DestStride,bool bLockWithinMiptail)
 {
-	DYNAMIC_CAST_METGALRESOURCE(Texture2DArray,Texture);
+	DYNAMIC_CAST_METALRESOURCE(Texture2DArray,Texture);
 	return Texture->Surface.Lock(MipIndex, TextureIndex, LockMode, DestStride);
 }
 
 void FMetalDynamicRHI::RHIUnlockTexture2DArray(FTexture2DArrayRHIParamRef TextureRHI,uint32 TextureIndex,uint32 MipIndex,bool bLockWithinMiptail)
 {
-	DYNAMIC_CAST_METGALRESOURCE(Texture2DArray,Texture);
+	DYNAMIC_CAST_METALRESOURCE(Texture2DArray,Texture);
 	Texture->Surface.Unlock(MipIndex, TextureIndex);
 }
 
 void FMetalDynamicRHI::RHIUpdateTexture2D(FTexture2DRHIParamRef TextureRHI, uint32 MipIndex, const struct FUpdateTextureRegion2D& UpdateRegion, uint32 SourcePitch, const uint8* SourceData)
 {	
-	DYNAMIC_CAST_METGALRESOURCE(Texture3D,Texture);
+	DYNAMIC_CAST_METALRESOURCE(Texture3D,Texture);
 
 }
 
 void FMetalDynamicRHI::RHIUpdateTexture3D(FTexture3DRHIParamRef TextureRHI,uint32 MipIndex,const FUpdateTextureRegion3D& UpdateRegion,uint32 SourceRowPitch,uint32 SourceDepthPitch,const uint8* SourceData)
 {	
-	DYNAMIC_CAST_METGALRESOURCE(Texture3D,Texture);
+	DYNAMIC_CAST_METALRESOURCE(Texture3D,Texture);
 
 }
 
@@ -517,14 +556,14 @@ FTextureCubeRHIRef FMetalDynamicRHI::RHICreateTextureCubeArray(uint32 Size, uint
 
 void* FMetalDynamicRHI::RHILockTextureCubeFace(FTextureCubeRHIParamRef TextureCubeRHI,uint32 FaceIndex,uint32 ArrayIndex,uint32 MipIndex,EResourceLockMode LockMode,uint32& DestStride,bool bLockWithinMiptail)
 {
-	DYNAMIC_CAST_METGALRESOURCE(TextureCube,TextureCube);
+	DYNAMIC_CAST_METALRESOURCE(TextureCube,TextureCube);
 	uint32 MetalFace = GetMetalCubeFace((ECubeFace)FaceIndex);
 	return TextureCube->Surface.Lock(MipIndex, FaceIndex + 6 * ArrayIndex, LockMode, DestStride);
 }
 
 void FMetalDynamicRHI::RHIUnlockTextureCubeFace(FTextureCubeRHIParamRef TextureCubeRHI,uint32 FaceIndex,uint32 ArrayIndex,uint32 MipIndex,bool bLockWithinMiptail)
 {
-	DYNAMIC_CAST_METGALRESOURCE(TextureCube,TextureCube);
+	DYNAMIC_CAST_METALRESOURCE(TextureCube,TextureCube);
 	uint32 MetalFace = GetMetalCubeFace((ECubeFace)FaceIndex);
 	TextureCube->Surface.Unlock(MipIndex, FaceIndex + ArrayIndex * 6);
 }
diff --git a/Engine/Source/Runtime/IOS/MetalRHI/Private/MetalUAV.cpp b/Engine/Source/Runtime/IOS/MetalRHI/Private/MetalUAV.cpp
--- a/Engine/Source/Runtime/IOS/MetalRHI/Private/MetalUAV.cpp
+++ b/Engine/Source/Runtime/IOS/MetalRHI/Private/MetalUAV.cpp
@@ -13,7 +13,7 @@ FMetalShaderResourceView::~FMetalShaderResourceView()
 
 FUnorderedAccessViewRHIRef FMetalDynamicRHI::RHICreateUnorderedAccessView(FStructuredBufferRHIParamRef StructuredBufferRHI, bool bUseUAVCounter, bool bAppendBuffer)
 {
-	DYNAMIC_CAST_METGALRESOURCE(StructuredBuffer,StructuredBuffer);
+	DYNAMIC_CAST_METALRESOURCE(StructuredBuffer,StructuredBuffer);
 
 	// create the UAV buffer to point to the structured buffer's memory
 	FMetalUnorderedAccessView* UAV = new FMetalUnorderedAccessView;
@@ -35,7 +35,7 @@ FUnorderedAccessViewRHIRef FMetalDynamicRHI::RHICreateUnorderedAccessView(FTextu
 
 FUnorderedAccessViewRHIRef FMetalDynamicRHI::RHICreateUnorderedAccessView(FVertexBufferRHIParamRef VertexBufferRHI, uint8 Format)
 {
-	DYNAMIC_CAST_METGALRESOURCE(VertexBuffer, VertexBuffer);
+	DYNAMIC_CAST_METALRESOURCE(VertexBuffer, VertexBuffer);
 
 	// create the UAV buffer to point to the structured buffer's memory
 	FMetalUnorderedAccessView* UAV = new FMetalUnorderedAccessView;
@@ -46,7 +46,7 @@ FUnorderedAccessViewRHIRef FMetalDynamicRHI::RHICreateUnorderedAccessView(FVerte
 
 FShaderResourceViewRHIRef FMetalDynamicRHI::RHICreateShaderResourceView(FStructuredBufferRHIParamRef StructuredBufferRHI)
 {
-	DYNAMIC_CAST_METGALRESOURCE(StructuredBuffer, StructuredBuffer);
+	DYNAMIC_CAST_METALRESOURCE(StructuredBuffer, StructuredBuffer);
 
 	FMetalShaderResourceView* SRV = new FMetalShaderResourceView;
 	return SRV;
@@ -54,7 +54,7 @@ FShaderResourceViewRHIRef FMetalDynamicRHI::RHICreateShaderResourceView(FStructu
 
 FShaderResourceViewRHIRef FMetalDynamicRHI::RHICreateShaderResourceView(FVertexBufferRHIParamRef VertexBufferRHI, uint32 Stride, uint8 Format)
 {
-	DYNAMIC_CAST_METGALRESOURCE(VertexBuffer, VertexBuffer);
+	DYNAMIC_CAST_METALRESOURCE(VertexBuffer, VertexBuffer);
 
 	FMetalShaderResourceView* SRV = new FMetalShaderResourceView;
 	SRV->SourceVertexBuffer = VertexBuffer;
@@ -77,7 +77,7 @@ FShaderResourceViewRHIRef FMetalDynamicRHI::RHICreateShaderResourceView(FTexture
 
 void FMetalDynamicRHI::RHIClearUAV(FUnorderedAccessViewRHIParamRef UnorderedAccessViewRHI, const uint32* Values)
 {
-	DYNAMIC_CAST_METGALRESOURCE(UnorderedAccessView, UnorderedAccessView);
+	DYNAMIC_CAST_METALRESOURCE(UnorderedAccessView, UnorderedAccessView);
 
 }
 
diff --git a/Engine/Source/Runtime/IOS/MetalRHI/Private/MetalUniformBuffer.cpp b/Engine/Source/Runtime/IOS/MetalRHI/Private/MetalUniformBuffer.cpp
--- a/Engine/Source/Runtime/IOS/MetalRHI/Private/MetalUniformBuffer.cpp
+++ b/Engine/Source/Runtime/IOS/MetalRHI/Private/MetalUniformBuffer.cpp
@@ -240,10 +240,9 @@ FMetalUniformBuffer::FMetalUniformBuffer(const void* Contents, const FRHIUniform
 FMetalUniformBuffer::~FMetalUniformBuffer()
 {
 	// don't need to free the ring buffer!
-	if (Buffer != FMetalManager::Get()->GetRingBuffer())
+	if (Buffer != nil && Buffer != FMetalManager::Get()->GetRingBuffer())
 	{
 		AddNewlyFreedBufferToUniformBufferPool(Buffer, Offset, GetSize());
-		//	FMetalManager::ReleaseObject(Buffer);
 	}
 }
 
diff --git a/Engine/Source/Runtime/IOS/MetalRHI/Private/MetalVertexBuffer.cpp b/Engine/Source/Runtime/IOS/MetalRHI/Private/MetalVertexBuffer.cpp
--- a/Engine/Source/Runtime/IOS/MetalRHI/Private/MetalVertexBuffer.cpp
+++ b/Engine/Source/Runtime/IOS/MetalRHI/Private/MetalVertexBuffer.cpp
@@ -79,7 +79,7 @@ FVertexBufferRHIRef FMetalDynamicRHI::RHICreateVertexBuffer(uint32 Size, uint32
 
 void* FMetalDynamicRHI::RHILockVertexBuffer(FVertexBufferRHIParamRef VertexBufferRHI, uint32 Offset, uint32 Size, EResourceLockMode LockMode)
 {
-	DYNAMIC_CAST_METGALRESOURCE(VertexBuffer,VertexBuffer);
+	DYNAMIC_CAST_METALRESOURCE(VertexBuffer,VertexBuffer);
 
 	// default to vertex buffer memory
 	return (uint8*)VertexBuffer->Lock(LockMode, Size) + Offset;
@@ -87,7 +87,7 @@ void* FMetalDynamicRHI::RHILockVertexBuffer(FVertexBufferRHIParamRef VertexBuffe
 
 void FMetalDynamicRHI::RHIUnlockVertexBuffer(FVertexBufferRHIParamRef VertexBufferRHI)
 {
-	DYNAMIC_CAST_METGALRESOURCE(VertexBuffer,VertexBuffer);
+	DYNAMIC_CAST_METALRESOURCE(VertexBuffer,VertexBuffer);
 
 	VertexBuffer->Unlock();
 }
diff --git a/Engine/Source/Runtime/IOS/MetalRHI/Private/MetalViewport.cpp b/Engine/Source/Runtime/IOS/MetalRHI/Private/MetalViewport.cpp
--- a/Engine/Source/Runtime/IOS/MetalRHI/Private/MetalViewport.cpp
+++ b/Engine/Source/Runtime/IOS/MetalRHI/Private/MetalViewport.cpp
@@ -30,7 +30,7 @@ void FMetalDynamicRHI::RHIResizeViewport(FViewportRHIParamRef ViewportRHI,uint32
 {
 	check( IsInGameThread() );
 
-	DYNAMIC_CAST_METGALRESOURCE(Viewport,Viewport);
+	DYNAMIC_CAST_METALRESOURCE(Viewport,Viewport);
 }
 
 void FMetalDynamicRHI::RHITick( float DeltaTime )
diff --git a/Engine/Source/Runtime/IOS/MetalRHI/Public/MetalRHI.h b/Engine/Source/Runtime/IOS/MetalRHI/Public/MetalRHI.h
--- a/Engine/Source/Runtime/IOS/MetalRHI/Public/MetalRHI.h
+++ b/Engine/Source/Runtime/IOS/MetalRHI/Public/MetalRHI.h
@@ -9,7 +9,7 @@
 DECLARE_LOG_CATEGORY_EXTERN(LogMetal, Display, All);
 
 /** This is a macro that casts a dynamically bound RHI reference to the appropriate type. */
-#define DYNAMIC_CAST_METGALRESOURCE(Type,Name) \
+#define DYNAMIC_CAST_METALRESOURCE(Type,Name) \
 	FMetal##Type* Name = (FMetal##Type*)Name##RHI;
 
 // Metal RHI public headers.
diff --git a/Engine/Source/Runtime/IOS/MetalRHI/Public/MetalResources.h b/Engine/Source/Runtime/IOS/MetalRHI/Public/MetalResources.h
--- a/Engine/Source/Runtime/IOS/MetalRHI/Public/MetalResources.h
+++ b/Engine/Source/Runtime/IOS/MetalRHI/Public/MetalResources.h
@@ -171,8 +171,10 @@ public:
 
 
 	EPixelFormat PixelFormat;
+	uint8 FormatKey;
 	id<MTLTexture> Texture;
     id<MTLTexture> MSAATexture;
+	id<MTLTexture> StencilTexture;
 	uint32 SizeX, SizeY, SizeZ;
 	bool bIsCubemap;
 	
@@ -182,6 +184,11 @@ public:
 
 	// how much memory is allocated for this texture
 	uint64 TotalTextureSize;
+
+private:
+	// a mapping of formats to small key, to fit into the PipelineState hash key
+	static TMap<MTLPixelFormat, uint8> PixelFormatKeyMap;
+	static uint8 NextKey;
 };
 
 class FMetalTexture2D : public FRHITexture2D
@@ -272,8 +279,8 @@ public:
 	// Where in the ring buffer to store the result
 	uint32 Offset;
 
-	// The command buffer this query is being rendered in - used to make sure the command buffer is complete whewn we read the result
-	uint64 CommandBufferIndex;
+	// This event will be fired when the command buffer that rendered the OQ mesh is completed by the GPU
+	FEvent* QueryCompleteEvent;
 };
 
 /** Index buffer resource class that stores stride information. */
diff --git a/Engine/Source/Runtime/RHI/Private/RHI.cpp b/Engine/Source/Runtime/RHI/Private/RHI.cpp
--- a/Engine/Source/Runtime/RHI/Private/RHI.cpp
+++ b/Engine/Source/Runtime/RHI/Private/RHI.cpp
@@ -278,6 +278,7 @@ static FName NAME_OPENGL_ES2(TEXT("GLSL_ES2"));
 static FName NAME_OPENGL_ES2_WEBGL(TEXT("GLSL_ES2_WEBGL"));
 static FName NAME_OPENGL_ES2_IOS(TEXT("GLSL_ES2_IOS"));
 static FName NAME_SF_METAL(TEXT("SF_METAL"));
+static FName NAME_SF_METAL_MRT(TEXT("SF_METAL_MRT"));
 static FName NAME_GLSL_310_ES_EXT(TEXT("GLSL_310_ES_EXT"));
 
 FName LegacyShaderPlatformToShaderFormat(EShaderPlatform Platform)
@@ -310,6 +311,8 @@ FName LegacyShaderPlatformToShaderFormat(EShaderPlatform Platform)
 		return NAME_OPENGL_ES2_IOS;
 	case SP_METAL:
 		return NAME_SF_METAL;
+	case SP_METAL_MRT:
+		return NAME_SF_METAL_MRT;
 	case SP_OPENGL_ES31_EXT:
 		return NAME_GLSL_310_ES_EXT;
 
@@ -334,6 +337,7 @@ EShaderPlatform ShaderFormatToLegacyShaderPlatform(FName ShaderFormat)
 	if (ShaderFormat == NAME_OPENGL_ES2_WEBGL)	return SP_OPENGL_ES2_WEBGL;
 	if (ShaderFormat == NAME_OPENGL_ES2_IOS)	return SP_OPENGL_ES2_IOS;
 	if (ShaderFormat == NAME_SF_METAL)			return SP_METAL;
+	if (ShaderFormat == NAME_SF_METAL_MRT)		return SP_METAL_MRT;
 	if (ShaderFormat == NAME_GLSL_310_ES_EXT)	return SP_OPENGL_ES31_EXT;
 	return SP_NumPlatforms;
 }
diff --git a/Engine/Source/Runtime/RHI/Public/RHIDefinitions.h b/Engine/Source/Runtime/RHI/Public/RHIDefinitions.h
--- a/Engine/Source/Runtime/RHI/Public/RHIDefinitions.h
+++ b/Engine/Source/Runtime/RHI/Public/RHIDefinitions.h
@@ -42,9 +42,10 @@ enum EShaderPlatform
 	SP_OPENGL_ES2_IOS	= 10,
 	SP_METAL			= 11,
 	SP_OPENGL_SM4_MAC	= 12,
-	SP_OPENGL_ES31_EXT	= 13,
+	SP_METAL_MRT		= 13,
+	SP_OPENGL_ES31_EXT	= 14,
 
-	SP_NumPlatforms		= 14,
+	SP_NumPlatforms		= 15,
 	SP_NumBits			= 5,
 };
 static_assert(SP_NumPlatforms <= (1 << SP_NumBits), "SP_NumPlatforms will not fit on SP_NumBits");
@@ -579,6 +580,7 @@ inline ERHIFeatureLevel::Type GetMaxSupportedFeatureLevel(EShaderPlatform InShad
 	case SP_PCD3D_SM4:
 	case SP_OPENGL_SM4:
 	case SP_OPENGL_SM4_MAC:
+	case SP_METAL_MRT:
 		return ERHIFeatureLevel::SM4;
 	case SP_PCD3D_ES2:
 	case SP_OPENGL_PCES2:
@@ -624,6 +626,8 @@ inline bool IsFeatureLevelSupported(EShaderPlatform InShaderPlatform, ERHIFeatur
 		return InFeatureLevel <= ERHIFeatureLevel::SM5;
 	case SP_METAL: 
 		return InFeatureLevel <= ERHIFeatureLevel::ES2;
+	case SP_METAL_MRT:
+		return InFeatureLevel <= ERHIFeatureLevel::SM4;
 	case SP_OPENGL_ES31_EXT:
 		return InFeatureLevel <= ERHIFeatureLevel::SM5;
 	default:
@@ -655,7 +659,17 @@ inline bool RHISupportsInstancing(const EShaderPlatform Platform)
 inline bool RHISupportsSeparateMSAAAndResolveTextures(const EShaderPlatform Platform)
 {
 	// Metal needs to handle MSAA and resolve textures internally (unless RHICreateTexture2D was changed to take an optional resolve target)
-	return Platform != SP_METAL;
+	return Platform != SP_METAL && Platform != SP_METAL_MRT;
+}
+
+inline bool RHISupportsGeometryShaders(const EShaderPlatform Platform)
+{
+	return IsFeatureLevelSupported(Platform, ERHIFeatureLevel::SM4) && Platform != SP_METAL_MRT;
+}
+
+inline bool RHIHasTiledGPU(const EShaderPlatform Platform)
+{
+	return Platform == SP_METAL_MRT|| Platform == SP_METAL || Platform == SP_OPENGL_ES2_IOS || Platform == SP_OPENGL_ES2;
 }
 
 inline uint32 GetFeatureLevelMaxTextureSamplers(ERHIFeatureLevel::Type FeatureLevel)
diff --git a/Engine/Source/Runtime/RHI/Public/RHIResources.h b/Engine/Source/Runtime/RHI/Public/RHIResources.h
--- a/Engine/Source/Runtime/RHI/Public/RHIResources.h
+++ b/Engine/Source/Runtime/RHI/Public/RHIResources.h
@@ -619,6 +619,16 @@ public:
 		LoadAction(InLoadAction),
 		StoreAction(InStoreAction)
 	{}
+
+	bool operator==(const FRHIRenderTargetView& Other) const
+	{
+		return 
+			Texture == Other.Texture &&
+			MipIndex == Other.MipIndex &&
+			ArraySliceIndex == Other.ArraySliceIndex &&
+			LoadAction == Other.LoadAction &&
+			StoreAction == Other.StoreAction;
+	}
 };
 
 class FRHIDepthRenderTargetView
@@ -626,32 +636,52 @@ class FRHIDepthRenderTargetView
 public:
 	FTextureRHIParamRef Texture;
 
-	ERenderTargetLoadAction LoadAction;
-	ERenderTargetStoreAction StoreAction;
+	ERenderTargetLoadAction DepthLoadAction;
+	ERenderTargetStoreAction DepthStoreAction;
+	ERenderTargetLoadAction StencilLoadAction;
+	ERenderTargetStoreAction StencilStoreAction;
 
 	FRHIDepthRenderTargetView() :
 		Texture(nullptr),
-		LoadAction(ERenderTargetLoadAction::EClear),
-		StoreAction(ERenderTargetStoreAction::ENoAction)
-	{}
-
-	FRHIDepthRenderTargetView(const FRHIDepthRenderTargetView& Other) :
-		Texture(Other.Texture),
-		LoadAction(Other.LoadAction),
-		StoreAction(Other.StoreAction)
+		DepthLoadAction(ERenderTargetLoadAction::EClear),
+		DepthStoreAction(ERenderTargetStoreAction::ENoAction),
+		StencilLoadAction(ERenderTargetLoadAction::EClear),
+		StencilStoreAction(ERenderTargetStoreAction::ENoAction)
 	{}
 
 	FRHIDepthRenderTargetView(FTextureRHIParamRef InTexture) :
 		Texture(InTexture),
-		LoadAction(ERenderTargetLoadAction::EClear),
-		StoreAction(ERenderTargetStoreAction::ENoAction)
+		DepthLoadAction(ERenderTargetLoadAction::EClear),
+		DepthStoreAction(ERenderTargetStoreAction::ENoAction),
+		StencilLoadAction(ERenderTargetLoadAction::EClear),
+		StencilStoreAction(ERenderTargetStoreAction::ENoAction)
 	{}
 
 	FRHIDepthRenderTargetView(FTextureRHIParamRef InTexture, ERenderTargetLoadAction InLoadAction, ERenderTargetStoreAction InStoreAction) :
 		Texture(InTexture),
-		LoadAction(InLoadAction),
-		StoreAction(InStoreAction)
+		DepthLoadAction(InLoadAction),
+		DepthStoreAction(InStoreAction),
+		StencilLoadAction(InLoadAction),
+		StencilStoreAction(InStoreAction)
+	{}
+
+	FRHIDepthRenderTargetView(FTextureRHIParamRef InTexture, ERenderTargetLoadAction InDepthLoadAction, ERenderTargetStoreAction InDepthStoreAction, ERenderTargetLoadAction InStencilLoadAction, ERenderTargetStoreAction InStencilStoreAction) :
+		Texture(InTexture),
+		DepthLoadAction(InDepthLoadAction),
+		DepthStoreAction(InDepthStoreAction),
+		StencilLoadAction(InStencilLoadAction),
+		StencilStoreAction(InStencilStoreAction)
 	{}
+
+	bool operator==(const FRHIDepthRenderTargetView& Other)
+	{
+		return
+			Texture == Other.Texture &&
+			DepthLoadAction == Other.DepthLoadAction &&
+			DepthStoreAction == Other.DepthStoreAction &&
+			StencilLoadAction == Other.StencilLoadAction &&
+			StencilStoreAction == Other.StencilStoreAction;
+	}
 };
 
 class FRHISetRenderTargetsInfo
@@ -681,21 +711,29 @@ public:
 
 	FRHISetRenderTargetsInfo(int32 InNumColorRenderTargets, const FRHIRenderTargetView* InColorRenderTargets, const FRHIDepthRenderTargetView& InDepthStencilRenderTarget) :
 		NumColorRenderTargets(InNumColorRenderTargets),
-		bClearColor(false),
+		bClearColor(InNumColorRenderTargets > 0 && InColorRenderTargets[0].LoadAction == ERenderTargetLoadAction::EClear),
 		DepthStencilRenderTarget(InDepthStencilRenderTarget),
-		bClearDepth(false),
-		bClearStencil(false)
+		bClearDepth(InDepthStencilRenderTarget.DepthLoadAction == ERenderTargetLoadAction::EClear),
+		bClearStencil(InDepthStencilRenderTarget.StencilLoadAction == ERenderTargetLoadAction::EClear)
 	{
 		check(InNumColorRenderTargets <= 0 || InColorRenderTargets);
 		for (int32 Index = 0; Index < InNumColorRenderTargets; ++Index)
 		{
 			ColorRenderTarget[Index] = InColorRenderTargets[Index];
+			ClearColors[Index] = FLinearColor(0, 0, 0, 1);
 		}
 	}
-
+	// @todo metal mrt: This can go away after all the cleanup is done
 	void SetClearDepthStencil(bool bInClearDepth, float InDepthClear, bool bInClearStencil = false, uint32 InClearStencil = 0)
 	{
-		DepthStencilRenderTarget.LoadAction = ERenderTargetLoadAction::EClear;
+		if (bInClearDepth)
+		{
+			DepthStencilRenderTarget.DepthLoadAction = ERenderTargetLoadAction::EClear;
+		}
+		if (bInClearStencil)
+		{
+			DepthStencilRenderTarget.StencilLoadAction = ERenderTargetLoadAction::EClear;
+		}
 		bClearDepth = bInClearDepth;
 		DepthClearValue = InDepthClear;
 		bClearStencil = bInClearStencil;
