UE4 - Metal SetSRV for buffers

diff --git a/Engine/Source/Runtime/IOS/MetalRHI/Private/MetalCommands.cpp b/Engine/Source/Runtime/IOS/MetalRHI/Private/MetalCommands.cpp
--- a/Engine/Source/Runtime/IOS/MetalRHI/Private/MetalCommands.cpp
+++ b/Engine/Source/Runtime/IOS/MetalRHI/Private/MetalCommands.cpp
@@ -147,7 +147,22 @@ void FMetalDynamicRHI::RHISetShaderTexture(FComputeShaderRHIParamRef ComputeShad
 
 void FMetalDynamicRHI::RHISetShaderResourceViewParameter(FVertexShaderRHIParamRef VertexShaderRHI, uint32 TextureIndex, FShaderResourceViewRHIParamRef SRVRHI)
 {
-	NOT_SUPPORTED("RHISetShaderResourceViewParameter");
+	DYNAMIC_CAST_METALRESOURCE(ShaderResourceView, SRV);
+
+	FRHITexture* Texture = SRV->SourceTexture.GetReference();
+	if (Texture)
+	{
+		FMetalSurface& Surface = GetMetalSurfaceFromRHITexture(Texture);
+		[FMetalManager::GetContext() setVertexTexture:Surface.Texture atIndex:TextureIndex];
+	}
+	else
+	{
+		FMetalVertexBuffer* VB = SRV->SourceVertexBuffer.GetReference();
+		if (VB)
+		{
+			[FMetalManager::GetContext() setVertexBuffer:VB->Buffer offset:VB->Offset atIndex:TextureIndex];
+		}
+	}
 }
 
 void FMetalDynamicRHI::RHISetShaderResourceViewParameter(FHullShaderRHIParamRef HullShaderRHI,uint32 TextureIndex,FShaderResourceViewRHIParamRef SRVRHI)
@@ -169,8 +184,20 @@ void FMetalDynamicRHI::RHISetShaderResourceViewParameter(FPixelShaderRHIParamRef
 {
 	DYNAMIC_CAST_METALRESOURCE(ShaderResourceView, SRV);
 
-	FMetalSurface& Surface = GetMetalSurfaceFromRHITexture(SRV->SourceTexture.GetReference());
-	[FMetalManager::GetContext() setFragmentTexture:Surface.Texture atIndex:TextureIndex];
+	FRHITexture* Texture = SRV->SourceTexture.GetReference();
+	if (Texture)
+	{
+		FMetalSurface& Surface = GetMetalSurfaceFromRHITexture(Texture);
+		[FMetalManager::GetContext() setFragmentTexture:Surface.Texture atIndex:TextureIndex];
+	}
+	else
+	{
+		FMetalVertexBuffer* VB = SRV->SourceVertexBuffer.GetReference();
+		if (VB)
+		{
+			[FMetalManager::GetContext() setFragmentBuffer:VB->Buffer offset:VB->Offset atIndex:TextureIndex];
+		}
+	}
 }
 
 void FMetalDynamicRHI::RHISetShaderResourceViewParameter(FComputeShaderRHIParamRef ComputeShaderRHI,uint32 TextureIndex,FShaderResourceViewRHIParamRef SRVRHI)
