bump the present counter to better hide some materials that get hiddne
#zen

diff --git a/Engine/Source/Runtime/IOS/MetalRHI/Private/MetalManager.cpp b/Engine/Source/Runtime/IOS/MetalRHI/Private/MetalManager.cpp
--- a/Engine/Source/Runtime/IOS/MetalRHI/Private/MetalManager.cpp
+++ b/Engine/Source/Runtime/IOS/MetalRHI/Private/MetalManager.cpp
@@ -377,7 +377,7 @@ void FMetalManager::EndFrame(bool bPresent)
 	// enqueue a present if desired
 	if (CurrentDrawable)
 	{
-		if (bPresent && GFrameCounter > 3)
+		if (bPresent && GFrameCounter > 30)
 		{
 			[CurrentCommandBuffer waitUntilScheduled];
 			
