- Fixed bad render target format hash in Metal

diff --git a/Engine/Source/Runtime/IOS/MetalRHI/Private/MetalManager.cpp b/Engine/Source/Runtime/IOS/MetalRHI/Private/MetalManager.cpp
--- a/Engine/Source/Runtime/IOS/MetalRHI/Private/MetalManager.cpp
+++ b/Engine/Source/Runtime/IOS/MetalRHI/Private/MetalManager.cpp
@@ -556,11 +556,14 @@ void FMetalManager::SetRenderTargetsInfo(const FRHISetRenderTargetsInfo& RenderT
 
 	for (uint32 RenderTargetIndex = 0; RenderTargetIndex < MaxMetalRenderTargets; RenderTargetIndex++)
 	{
+		// default to invalid
+		uint8 FormatKey = 0;
 		// only try to set it if it was one that was set (ie less than CurrentNumRenderTargets)
 		if (RenderTargetIndex < RenderTargetsInfo.NumColorRenderTargets && RenderTargetsInfo.ColorRenderTarget[RenderTargetIndex].Texture != nullptr)
 		{
 			const FRHIRenderTargetView& RenderTargetView = RenderTargetsInfo.ColorRenderTarget[RenderTargetIndex];
 			FMetalSurface& Surface = GetMetalSurfaceFromRHITexture(RenderTargetView.Texture);
+			FormatKey = Surface.FormatKey;
 		
 			// if this is the back buffer, make sure we have a usable drawable
 			ConditionalUpdateBackBuffer(Surface);
@@ -611,7 +614,7 @@ void FMetalManager::SetRenderTargetsInfo(const FRHISetRenderTargetsInfo& RenderT
 		}
 
 		// update the hash no matter what case (null, unused, used)
-		SET_HASH(RTBitOffsets[RenderTargetIndex], NUMBITS_RENDER_TARGET_FORMAT, Pipeline.RenderTargetFormatKeys[RenderTargetIndex]);
+		SET_HASH(RTBitOffsets[RenderTargetIndex], NUMBITS_RENDER_TARGET_FORMAT, FormatKey);
 	}
 
 	// default to invalid
diff --git a/Engine/Source/Runtime/IOS/MetalRHI/Private/MetalManager.h b/Engine/Source/Runtime/IOS/MetalRHI/Private/MetalManager.h
--- a/Engine/Source/Runtime/IOS/MetalRHI/Private/MetalManager.h
+++ b/Engine/Source/Runtime/IOS/MetalRHI/Private/MetalManager.h
@@ -40,7 +40,6 @@ struct FPipelineShadow
 	id<MTLRenderPipelineState> CreatePipelineStateForBoundShaderState(FMetalBoundShaderState* BSS) const;
 
 	MTLRenderPipelineColorAttachmentDescriptor* RenderTargets[MaxMetalRenderTargets];
-	uint8 RenderTargetFormatKeys[MaxMetalRenderTargets];
 	MTLPixelFormat DepthTargetFormat;
 	MTLPixelFormat StencilTargetFormat;
     uint32 SampleCount;
