- Fixed when Metal switches render targets

diff --git a/Engine/Source/Runtime/IOS/MetalRHI/Private/MetalManager.cpp b/Engine/Source/Runtime/IOS/MetalRHI/Private/MetalManager.cpp
--- a/Engine/Source/Runtime/IOS/MetalRHI/Private/MetalManager.cpp
+++ b/Engine/Source/Runtime/IOS/MetalRHI/Private/MetalManager.cpp
@@ -515,11 +515,35 @@ bool FMetalManager::NeedsToSetRenderTarget(const FRHISetRenderTargetsInfo& Rende
 			const FRHIRenderTargetView& RenderTargetView = RenderTargetsInfo.ColorRenderTarget[RenderTargetIndex];
 			const FRHIRenderTargetView& PreviousRenderTargetView = PreviousRenderTargetsInfo.ColorRenderTarget[RenderTargetIndex];
 
-			if (!(RenderTargetView == PreviousRenderTargetView))
+            // handle simple case of switching textures or mip/slice
+            if (RenderTargetView.Texture != PreviousRenderTargetView.Texture ||
+                RenderTargetView.MipIndex != PreviousRenderTargetView.MipIndex ||
+                RenderTargetView.ArraySliceIndex != PreviousRenderTargetView.ArraySliceIndex)
+            {
+                bAllChecksPassed = false;
+                break;
+            }
+            
+			// it's non-trivial when we need to switch based on load/store action:
+			// LoadAction - it only matters what we are switching to in the new one
+			//    If we switch to Load, no need to switch as we can re-use what we already have
+			//    If we switch to Clear, we have to always switch to a new RT to force the clear
+			//    If we switch to DontCare, there's definitely no need to switch
+			if (RenderTargetView.LoadAction == ERenderTargetLoadAction::EClear)
 			{
 				bAllChecksPassed = false;
 				break;
 			}
+			// StoreAction - this matters what the previous one was **In Spirit**
+			//    If we come from Store, we need to switch to a new RT to force the store
+			//    If we come from DontCare, then there's no need to switch
+			//    @todo metal: However, we basically only use Store now, and don't
+			//        care about intermediate results, only final, so we don't currently check the value
+//			if (PreviousRenderTargetView.StoreAction == ERenderTTargetStoreAction::EStore)
+//			{
+//				bAllChecksPassed = false;
+//				break;
+//			}
 		}
 	}
 
