switch to min IOS version of 6.1
removed cook compression from IOS settings as they are already in rendering settings
UE-5083
UE-5081
#ios

diff --git a/Engine/Source/Developer/iOS/IOSPlatformEditor/Private/IOSTargetSettingsCustomization.cpp b/Engine/Source/Developer/iOS/IOSPlatformEditor/Private/IOSTargetSettingsCustomization.cpp
--- a/Engine/Source/Developer/iOS/IOSPlatformEditor/Private/IOSTargetSettingsCustomization.cpp
+++ b/Engine/Source/Developer/iOS/IOSPlatformEditor/Private/IOSTargetSettingsCustomization.cpp
@@ -443,8 +443,8 @@ void FIOSTargetSettingsCustomization::OnPlistPropertyModified()
 	FString iOSVersionBody = TEXT("\n\t<string>");
 	switch (Settings.MinimumiOSVersion)
 	{
-	case EIOSVersion::IOS_6:
-		iOSVersionBody += TEXT("6.0");
+	case EIOSVersion::IOS_61:
+		iOSVersionBody += TEXT("6.1");
 		break;
 	case EIOSVersion::IOS_7:
 		iOSVersionBody += TEXT("7.0");
diff --git a/Engine/Source/Programs/UnrealBuildTool/IOS/IOSToolChain.cs b/Engine/Source/Programs/UnrealBuildTool/IOS/IOSToolChain.cs
--- a/Engine/Source/Programs/UnrealBuildTool/IOS/IOSToolChain.cs
+++ b/Engine/Source/Programs/UnrealBuildTool/IOS/IOSToolChain.cs
@@ -49,10 +49,10 @@ namespace UnrealBuildTool
 
 		/** Which version of the iOS to allow at build time */
 		[XmlConfig]
-		public static string BuildIOSVersion = "6.0";
+		public static string BuildIOSVersion = "6.1";
 
 		/** Which version of the iOS to allow at run time */
-		public static string RunTimeIOSVersion = "6.0";
+		public static string RunTimeIOSVersion = "6.1";
 
 		/** Which developer directory to root from */
 		[XmlConfig]
diff --git a/Engine/Source/Runtime/IOS/IOSRuntimeSettings/Classes/IOSRuntimeSettings.h b/Engine/Source/Runtime/IOS/IOSRuntimeSettings/Classes/IOSRuntimeSettings.h
--- a/Engine/Source/Runtime/IOS/IOSRuntimeSettings/Classes/IOSRuntimeSettings.h
+++ b/Engine/Source/Runtime/IOS/IOSRuntimeSettings/Classes/IOSRuntimeSettings.h
@@ -24,8 +24,8 @@ enum class EPowerUsageFrameRateLock : uint8
 UENUM()
 	enum class EIOSVersion : uint8
 {
-	/** iOS 6 */
-	IOS_6 = 6 UMETA(DisplayName="6.0"),
+	/** iOS 6.1 */
+	IOS_61 = 6 UMETA(DisplayName="6.1"),
 
 	/** iOS 7 */
 	IOS_7 = 7 UMETA(DisplayName="7.0"),
@@ -127,10 +127,6 @@ public:
 	UPROPERTY(GlobalConfig, EditAnywhere, Category = DeviceUsage)
 	uint32 bSupportsIPhone : 1;
 
-	// Texture compression value
-	UPROPERTY(GlobalConfig, EditAnywhere, Category = CookSettings)
-	uint32 CompressionLevel;
-
 #if WITH_EDITOR
 	// UObject interface
 	virtual void PostEditChangeProperty(struct FPropertyChangedEvent& PropertyChangedEvent) override;
diff --git a/Engine/Source/Runtime/IOS/IOSRuntimeSettings/Private/IOSRuntimeSettings.cpp b/Engine/Source/Runtime/IOS/IOSRuntimeSettings/Private/IOSRuntimeSettings.cpp
--- a/Engine/Source/Runtime/IOS/IOSRuntimeSettings/Private/IOSRuntimeSettings.cpp
+++ b/Engine/Source/Runtime/IOS/IOSRuntimeSettings/Private/IOSRuntimeSettings.cpp
@@ -16,8 +16,7 @@ UIOSRuntimeSettings::UIOSRuntimeSettings(const FObjectInitializer& ObjectInitial
     FrameRateLock = EPowerUsageFrameRateLock::PUFRL_30;
 	bSupportsIPad = true;
 	bSupportsIPhone = true;
-	MinimumiOSVersion = EIOSVersion::IOS_6;
-	CompressionLevel = 0;
+	MinimumiOSVersion = EIOSVersion::IOS_61;
 	bDevForArmV7 = true;
 	bDevForArm64 = false;
 	bDevForArmV7S = false;
