UE4 - Fix metal crash with null srv

diff --git a/Engine/Source/Runtime/IOS/MetalRHI/Private/MetalCommands.cpp b/Engine/Source/Runtime/IOS/MetalRHI/Private/MetalCommands.cpp
--- a/Engine/Source/Runtime/IOS/MetalRHI/Private/MetalCommands.cpp
+++ b/Engine/Source/Runtime/IOS/MetalRHI/Private/MetalCommands.cpp
@@ -114,14 +114,14 @@ void FMetalDynamicRHI::RHISetUAVParameter(FComputeShaderRHIParamRef ComputeShade
 void FMetalDynamicRHI::RHISetShaderTexture(FVertexShaderRHIParamRef VertexShaderRHI, uint32 TextureIndex, FTextureRHIParamRef NewTextureRHI)
 {
 	FMetalSurface* Surface = GetMetalSurfaceFromRHITexture(NewTextureRHI);
-    if (Surface != nullptr)
-    {
-        [FMetalManager::GetContext() setVertexTexture:Surface->Texture atIndex:TextureIndex];
-    }
-    else
-    {
-        [FMetalManager::GetContext() setVertexTexture:nil atIndex:TextureIndex];
-    }
+	if (Surface != nullptr)
+	{
+		[FMetalManager::GetContext() setVertexTexture:Surface->Texture atIndex:TextureIndex];
+	}
+	else
+	{
+		[FMetalManager::GetContext() setVertexTexture:nil atIndex:TextureIndex];
+	}
 }
 
 void FMetalDynamicRHI::RHISetShaderTexture(FHullShaderRHIParamRef HullShader, uint32 TextureIndex, FTextureRHIParamRef NewTextureRHI)
@@ -143,14 +143,14 @@ void FMetalDynamicRHI::RHISetShaderTexture(FGeometryShaderRHIParamRef GeometrySh
 void FMetalDynamicRHI::RHISetShaderTexture(FPixelShaderRHIParamRef PixelShader, uint32 TextureIndex, FTextureRHIParamRef NewTextureRHI)
 {
 	FMetalSurface* Surface = GetMetalSurfaceFromRHITexture(NewTextureRHI);
-    if (Surface != nullptr)
-    {
-        [FMetalManager::GetContext() setFragmentTexture:Surface->Texture atIndex:TextureIndex];
-    }
-    else
-    {
-        [FMetalManager::GetContext() setFragmentTexture:nil atIndex:TextureIndex];
-    }
+	if (Surface != nullptr)
+	{
+		[FMetalManager::GetContext() setFragmentTexture:Surface->Texture atIndex:TextureIndex];
+	}
+	else
+	{
+		[FMetalManager::GetContext() setFragmentTexture:nil atIndex:TextureIndex];
+	}
 }
 
 void FMetalDynamicRHI::RHISetShaderTexture(FComputeShaderRHIParamRef ComputeShader, uint32 TextureIndex, FTextureRHIParamRef NewTextureRHI)
@@ -162,27 +162,33 @@ void FMetalDynamicRHI::RHISetShaderTexture(FComputeShaderRHIParamRef ComputeShad
 void FMetalDynamicRHI::RHISetShaderResourceViewParameter(FVertexShaderRHIParamRef VertexShaderRHI, uint32 TextureIndex, FShaderResourceViewRHIParamRef SRVRHI)
 {
 	DYNAMIC_CAST_METALRESOURCE(ShaderResourceView, SRV);
-
-	FRHITexture* Texture = SRV->SourceTexture.GetReference();
-	if (Texture)
+	if (SRV)
 	{
-		FMetalSurface* Surface = GetMetalSurfaceFromRHITexture(Texture);
-		if (Surface != nullptr)
+		FRHITexture* Texture = SRV->SourceTexture.GetReference();
+		if (Texture)
 		{
-			[FMetalManager::GetContext() setVertexTexture:Surface->Texture atIndex : TextureIndex];
+			FMetalSurface* Surface = GetMetalSurfaceFromRHITexture(Texture);
+			if (Surface != nullptr)
+			{
+				[FMetalManager::GetContext() setVertexTexture:Surface->Texture atIndex : TextureIndex];
+			}
+			else
+			{
+				[FMetalManager::GetContext() setVertexTexture:nil atIndex : TextureIndex];
+			}
 		}
 		else
 		{
-			[FMetalManager::GetContext() setVertexTexture:nil atIndex : TextureIndex];
+			FMetalVertexBuffer* VB = SRV->SourceVertexBuffer.GetReference();
+			if (VB)
+			{
+				[FMetalManager::GetContext() setVertexBuffer:VB->Buffer offset:VB->Offset atIndex:TextureIndex];
+			}
 		}
 	}
 	else
 	{
-		FMetalVertexBuffer* VB = SRV->SourceVertexBuffer.GetReference();
-		if (VB)
-		{
-			[FMetalManager::GetContext() setVertexBuffer:VB->Buffer offset:VB->Offset atIndex:TextureIndex];
-		}
+		[FMetalManager::GetContext() setVertexTexture:nil atIndex : TextureIndex];
 	}
 }
 
@@ -204,27 +210,33 @@ void FMetalDynamicRHI::RHISetShaderResourceViewParameter(FGeometryShaderRHIParam
 void FMetalDynamicRHI::RHISetShaderResourceViewParameter(FPixelShaderRHIParamRef PixelShaderRHI,uint32 TextureIndex,FShaderResourceViewRHIParamRef SRVRHI)
 {
 	DYNAMIC_CAST_METALRESOURCE(ShaderResourceView, SRV);
-
-	FRHITexture* Texture = SRV->SourceTexture.GetReference();
-	if (Texture)
+	if (SRV)
 	{
-		FMetalSurface* Surface = GetMetalSurfaceFromRHITexture(Texture);
-		if (Surface != nullptr)
+		FRHITexture* Texture = SRV->SourceTexture.GetReference();
+		if (Texture)
 		{
-			[FMetalManager::GetContext() setFragmentTexture:Surface->Texture atIndex : TextureIndex];
+			FMetalSurface* Surface = GetMetalSurfaceFromRHITexture(Texture);
+			if (Surface != nullptr)
+			{
+				[FMetalManager::GetContext() setFragmentTexture:Surface->Texture atIndex : TextureIndex];
+			}
+			else
+			{
+				[FMetalManager::GetContext() setFragmentTexture:nil atIndex : TextureIndex];
+			}
 		}
 		else
 		{
-			[FMetalManager::GetContext() setFragmentTexture:nil atIndex : TextureIndex];
+			FMetalVertexBuffer* VB = SRV->SourceVertexBuffer.GetReference();
+			if (VB)
+			{
+				[FMetalManager::GetContext() setFragmentBuffer:VB->Buffer offset:VB->Offset atIndex:TextureIndex];
+			}
 		}
 	}
 	else
 	{
-		FMetalVertexBuffer* VB = SRV->SourceVertexBuffer.GetReference();
-		if (VB)
-		{
-			[FMetalManager::GetContext() setFragmentBuffer:VB->Buffer offset:VB->Offset atIndex:TextureIndex];
-		}
+		[FMetalManager::GetContext() setFragmentTexture:nil atIndex : TextureIndex];
 	}
 }
 
