UE4 - hlslcc - Metal - NumThreads for compute shaders

diff --git a/Engine/Source/Developer/iOS/MetalShaderFormat/Private/MetalBackend.cpp b/Engine/Source/Developer/iOS/MetalShaderFormat/Private/MetalBackend.cpp
--- a/Engine/Source/Developer/iOS/MetalShaderFormat/Private/MetalBackend.cpp
+++ b/Engine/Source/Developer/iOS/MetalShaderFormat/Private/MetalBackend.cpp
@@ -340,6 +340,11 @@ protected:
 	exec_list sampler_variables;
 	exec_list image_variables;
 
+	/** Attribute [numthreads(X,Y,Z)] */
+	int32 NumThreadsX;
+	int32 NumThreadsY;
+	int32 NumThreadsZ;
+
 	/** Track global instructions. */
 	struct global_ir : public exec_node
 	{
@@ -410,8 +415,10 @@ protected:
 			{
 				md_array_entry* entry = (md_array_entry*)iter.get();
 				if (entry->type == type)
+				{
 					return true;
 				}
+			}
 			md_array_entry* entry = new(mem_ctx) md_array_entry();
 			entry->type = type;
 			used_md_arrays.push_tail(entry);
@@ -450,7 +457,7 @@ protected:
 	/**
 	 * Add tabs/spaces for the current indentation level.
 	 */
-	void indent(void)
+	void indent()
 	{
 		for (int i = 0; i < indentation; i++)
 		{
@@ -853,6 +860,14 @@ protected:
 			indentation--;
 		}
 
+		// Copy the global attributes
+		if (sig->is_main)
+		{
+			NumThreadsX = sig->wg_size_x;
+			NumThreadsY = sig->wg_size_y;
+			NumThreadsZ = sig->wg_size_z;
+		}
+
 		indentation++;
 		foreach_iter(exec_list_iterator, iter, sig->body)
 		{
@@ -2365,6 +2380,11 @@ protected:
 				ralloc_asprintf_append(buffer, "\n");
 			}
 		}
+
+		if (Frequency == compute_shader)
+		{
+			ralloc_asprintf_append(buffer, "// @NumThreads: %d, %d, %d\n", this->NumThreadsX, this->NumThreadsY, this->NumThreadsZ);
+		}
 	}
 
 public:
@@ -2599,6 +2619,9 @@ bool FMetalCodeBackend::GenerateMain(EHlslShaderFrequency Frequency, const char*
 	// Call the original EntryPoint
 	MainSig->body.push_tail(new(ParseState) ir_call(EntryPointSig, EntryPointReturn, &ArgInstructions));
 	MainSig->body.append_list(&PostCallInstructions);
+	MainSig->wg_size_x = EntryPointSig->wg_size_x;
+	MainSig->wg_size_y = EntryPointSig->wg_size_y;
+	MainSig->wg_size_z = EntryPointSig->wg_size_z;
 
 	// Generate the Main() function
 	auto* MainFunction = new(ParseState)ir_function("Main");
diff --git a/Engine/Source/Developer/iOS/MetalShaderFormat/Private/MetalShaderCompiler.cpp b/Engine/Source/Developer/iOS/MetalShaderFormat/Private/MetalShaderCompiler.cpp
--- a/Engine/Source/Developer/iOS/MetalShaderFormat/Private/MetalShaderCompiler.cpp
+++ b/Engine/Source/Developer/iOS/MetalShaderFormat/Private/MetalShaderCompiler.cpp
@@ -114,6 +114,7 @@ static void BuildMetalShaderOutput(
 	DEF_PREFIX_STR(Samplers);
 	DEF_PREFIX_STR(UAVs);
 	DEF_PREFIX_STR(SamplerStates);
+	DEF_PREFIX_STR(NumThreads);
 	#undef DEF_PREFIX_STR
 
 	// Skip any comments that come before the signature.
@@ -601,6 +602,19 @@ static void BuildMetalShaderOutput(
 		}
 	}
 
+	if (FCStringAnsi::Strncmp(ShaderSource, NumThreadsPrefix, NumThreadsPrefixLen) == 0)
+	{
+		ShaderSource += NumThreadsPrefixLen;
+		Header.NumThreadsX = ParseNumber(ShaderSource);
+		verify(Match(ShaderSource, ','));
+		Match(ShaderSource, ' ');
+		Header.NumThreadsY = ParseNumber(ShaderSource);
+		verify(Match(ShaderSource, ','));
+		Match(ShaderSource, ' ');
+		Header.NumThreadsZ = ParseNumber(ShaderSource);
+		verify(Match(ShaderSource, '\n'));
+	}
+
 	// Build the SRT for this shader.
 	{
 		// Build the generic SRT for this shader.
diff --git a/Engine/Source/Developer/iOS/MetalShaderFormat/Private/MetalShaderFormat.cpp b/Engine/Source/Developer/iOS/MetalShaderFormat/Private/MetalShaderFormat.cpp
--- a/Engine/Source/Developer/iOS/MetalShaderFormat/Private/MetalShaderFormat.cpp
+++ b/Engine/Source/Developer/iOS/MetalShaderFormat/Private/MetalShaderFormat.cpp
@@ -15,7 +15,7 @@ class FMetalShaderFormat : public IShaderFormat
 {
 	enum
 	{
-		HEADER_VERSION = 10,
+		HEADER_VERSION = 11,
 	};
 	struct FVersion
 	{
diff --git a/Engine/Source/Runtime/IOS/MetalRHI/Private/MetalShaders.cpp b/Engine/Source/Runtime/IOS/MetalRHI/Private/MetalShaders.cpp
--- a/Engine/Source/Runtime/IOS/MetalRHI/Private/MetalShaders.cpp
+++ b/Engine/Source/Runtime/IOS/MetalRHI/Private/MetalShaders.cpp
@@ -48,8 +48,8 @@ static FMetalCompiledShaderCache& GetMetalCompiledShaderCache()
 }
 
 /** Initialization constructor. */
-template<typename BaseResourceType>
-TMetalBaseShader<BaseResourceType>::TMetalBaseShader(const TArray<uint8>& InCode)
+template<typename BaseResourceType, int32 ShaderType>
+TMetalBaseShader<BaseResourceType, ShaderType>::TMetalBaseShader(const TArray<uint8>& InCode)
 	: DirtyUniformBuffers(0)
 {
 	FMemoryReader Ar(InCode, true);
@@ -159,12 +159,21 @@ TMetalBaseShader<BaseResourceType>::TMetalBaseShader(const TArray<uint8>& InCode
 		Bindings = Header.Bindings;
 		BoundUniformBuffers.AddZeroed(Header.Bindings.NumUniformBuffers);
 		UniformBuffersCopyInfo = Header.UniformBuffersCopyInfo;
+
+		//@todo: Find better way...
+		if (ShaderType == SF_Compute)
+		{
+			auto* ComputeShader = (FMetalComputeShader*)this;
+			ComputeShader->NumThreadsX = Header.NumThreadsX;
+			ComputeShader->NumThreadsY = Header.NumThreadsY;
+			ComputeShader->NumThreadsZ = Header.NumThreadsZ;
+		}
 	}
 }
 
 /** Destructor */
-template<typename BaseResourceType>
-TMetalBaseShader<BaseResourceType>::~TMetalBaseShader()
+template<typename BaseResourceType, int32 ShaderType>
+TMetalBaseShader<BaseResourceType, ShaderType>::~TMetalBaseShader()
 {
 	UNTRACK_OBJECT(Function);
 	[Function release];
diff --git a/Engine/Source/Runtime/IOS/MetalRHI/Public/MetalResources.h b/Engine/Source/Runtime/IOS/MetalRHI/Public/MetalResources.h
--- a/Engine/Source/Runtime/IOS/MetalRHI/Public/MetalResources.h
+++ b/Engine/Source/Runtime/IOS/MetalRHI/Public/MetalResources.h
@@ -22,7 +22,7 @@ public:
 	FVertexDeclarationElementList Elements;
 
 	/** This is the layout for the vertex elements */
- 	MTLVertexDescriptor* Layout;
+	MTLVertexDescriptor* Layout;
 
 protected:
 	void GenerateLayout(const FVertexDeclarationElementList& Elements);
@@ -32,10 +32,11 @@ protected:
 
 
 /** This represents a vertex shader that hasn't been combined with a specific declaration to create a bound shader. */
-template<typename BaseResourceType /*, typename ShaderType */>
+template<typename BaseResourceType, int32 ShaderType>
 class TMetalBaseShader : public BaseResourceType, public IRefCountedObject
 {
 public:
+	enum { StaticFrequency = ShaderType };
 
 	/** Initialization constructor. */
 	TMetalBaseShader()
@@ -64,7 +65,7 @@ public:
 	// this is the compiler shader
 	id<MTLFunction> Function;
 
- 	/** External bindings for this shader. */
+	/** External bindings for this shader. */
 	FMetalShaderBindings Bindings;
 	TArray< TRefCountPtr<FRHIUniformBuffer> > BoundUniformBuffers;
 
@@ -73,19 +74,33 @@ public:
 
 	// List of memory copies from RHIUniformBuffer to packed uniforms
 	TArray<FMetalUniformBufferCopyInfo> UniformBuffersCopyInfo;
-    
+	
 	TArray<ANSICHAR> GlslCode;
-    NSString* GlslCodeNSString;
+	NSString* GlslCodeNSString;
 	const ANSICHAR*  GlslCodeString; // make it easier in VS to see shader code in debug mode; points to begin of GlslCode
 };
 
-typedef TMetalBaseShader<FRHIVertexShader> FMetalVertexShader;
-typedef TMetalBaseShader<FRHIPixelShader> FMetalPixelShader;
-typedef TMetalBaseShader<FRHIHullShader> FMetalHullShader;
-typedef TMetalBaseShader<FRHIDomainShader> FMetalDomainShader;
-typedef TMetalBaseShader<FRHIComputeShader> FMetalComputeShader;
-typedef TMetalBaseShader<FRHIGeometryShader> FMetalGeometryShader;
+typedef TMetalBaseShader<FRHIVertexShader, SF_Vertex> FMetalVertexShader;
+typedef TMetalBaseShader<FRHIPixelShader, SF_Pixel> FMetalPixelShader;
+typedef TMetalBaseShader<FRHIHullShader, SF_Hull> FMetalHullShader;
+typedef TMetalBaseShader<FRHIDomainShader, SF_Domain> FMetalDomainShader;
+typedef TMetalBaseShader<FRHIGeometryShader, SF_Geometry> FMetalGeometryShader;
 
+class FMetalComputeShader : public TMetalBaseShader<FRHIComputeShader, SF_Compute>
+{
+public:
+	FMetalComputeShader(const TArray<uint8>& InCode) :
+		TMetalBaseShader<FRHIComputeShader, SF_Compute>(InCode),
+		NumThreadsX(0),
+		NumThreadsY(0),
+		NumThreadsZ(0)
+	{
+	}
+
+	int32 NumThreadsX;
+	int32 NumThreadsY;
+	int32 NumThreadsZ;
+};
 
 /**
  * Combined shader state and vertex definition for rendering geometry. 
@@ -173,7 +188,7 @@ public:
 	EPixelFormat PixelFormat;
 	uint8 FormatKey;
 	id<MTLTexture> Texture;
-    id<MTLTexture> MSAATexture;
+	id<MTLTexture> MSAATexture;
 	id<MTLTexture> StencilTexture;
 	uint32 SizeX, SizeY, SizeZ;
 	bool bIsCubemap;
@@ -202,8 +217,8 @@ public:
 		, Surface(RRT_Texture2D, Format, SizeX, SizeY, 1, NumSamples, /*bArray=*/ false, 1, NumMips, Flags, BulkData)
 	{
 	}
-    
-    
+	
+	
 };
 
 class FMetalTexture2DArray : public FRHITexture2DArray
diff --git a/Engine/Source/Runtime/IOS/MetalRHI/Public/MetalShaderResources.h b/Engine/Source/Runtime/IOS/MetalRHI/Public/MetalShaderResources.h
--- a/Engine/Source/Runtime/IOS/MetalRHI/Public/MetalShaderResources.h
+++ b/Engine/Source/Runtime/IOS/MetalRHI/Public/MetalShaderResources.h
@@ -165,6 +165,10 @@ struct FMetalCodeHeader
 	uint32 Frequency;
 	FMetalShaderBindings Bindings;
 	TArray<FMetalUniformBufferCopyInfo> UniformBuffersCopyInfo;
+
+	uint8 NumThreadsX;
+	uint8 NumThreadsY;
+	uint8 NumThreadsZ;
 };
 
 inline FArchive& operator<<(FArchive& Ar, FMetalCodeHeader& Header)
@@ -192,5 +196,9 @@ inline FArchive& operator<<(FArchive& Ar, FMetalCodeHeader& Header)
 		}
 	}
 
+	Ar << Header.NumThreadsX;
+	Ar << Header.NumThreadsY;
+	Ar << Header.NumThreadsZ;
+
 	return Ar;
 }
