UE4 - Metal - hardware instancing (for instanced static meshes)

diff --git a/Engine/Shaders/LocalVertexFactory.usf b/Engine/Shaders/LocalVertexFactory.usf
--- a/Engine/Shaders/LocalVertexFactory.usf
+++ b/Engine/Shaders/LocalVertexFactory.usf
@@ -1,4 +1,4 @@
-// Copyright 1998-2014 Epic Games, Inc. All Rights Reserved...
+// Copyright 1998-2014 Epic Games, Inc. All Rights Reserved..
 
 /*=============================================================================
 	LocalVertexFactory.usf: Local vertex factory shader code.
@@ -59,10 +59,10 @@ struct FVertexFactoryInput
 {
 	float4	Position	: ATTRIBUTE0;
 #if METAL_PROFILE
-		//@todo-rco: FIXME!
-		float3	TangentX	: ATTRIBUTE1;
-		float4	TangentZ	: ATTRIBUTE2;
-		float4	Color		: ATTRIBUTE3;
+	//@todo-rco: FIXME!
+	float3	TangentX	: ATTRIBUTE1;
+	float4	TangentZ	: ATTRIBUTE2;
+	float4	Color		: ATTRIBUTE3;
 #else
 	half3	TangentX	: ATTRIBUTE1;
 	// TangentZ.w contains sign of tangent basis determinant
diff --git a/Engine/Source/Runtime/IOS/MetalRHI/Private/MetalVertexDeclaration.cpp b/Engine/Source/Runtime/IOS/MetalRHI/Private/MetalVertexDeclaration.cpp
--- a/Engine/Source/Runtime/IOS/MetalRHI/Private/MetalVertexDeclaration.cpp
+++ b/Engine/Source/Runtime/IOS/MetalRHI/Private/MetalVertexDeclaration.cpp
@@ -93,7 +93,7 @@ void FMetalVertexDeclaration::GenerateLayout(const FVertexDeclarationElementList
 		checkf(Element.Stride == 0 || Element.Offset + TranslateElementTypeToSize(Element.Type) <= Element.Stride, 
 			TEXT("Stream component is bigger than stride: Offset: %d, Size: %d [Type %d], Stride: %d"), Element.Offset, TranslateElementTypeToSize(Element.Type), (uint32)Element.Type, Element.Stride);
 
-		// we offset 6 buffers to leave space for uniform buffers
+		// Vertex & Constant buffers are set up in the same space, so add VB's from the top
 		uint32 ShaderBufferIndex = UNREAL_TO_METAL_BUFFER_INDEX(Element.StreamIndex);
 
 		// track the buffer stride, making sure all elements with the same buffer have the same stride
@@ -101,7 +101,7 @@ void FMetalVertexDeclaration::GenerateLayout(const FVertexDeclarationElementList
 		if (ExistingStride == NULL)
 		{
 			// handle 0 stride buffers
-			MTLVertexStepFunction Function = (Element.Stride == 0 ? MTLVertexStepFunctionConstant : MTLVertexStepFunctionPerVertex);
+			MTLVertexStepFunction Function = (Element.Stride == 0 ? MTLVertexStepFunctionConstant : (Element.bUseInstanceIndex ? MTLVertexStepFunctionPerInstance : MTLVertexStepFunctionPerVertex));
 			uint32 StepRate = (Element.Stride == 0 ? 0 : 1);
 			// even with MTLVertexStepFunctionConstant, it needs a non-zero stride (not sure why)
 			uint32 Stride = (Element.Stride == 0 ? 4 : Element.Stride);
diff --git a/Engine/Source/Runtime/RHI/Public/RHIDefinitions.h b/Engine/Source/Runtime/RHI/Public/RHIDefinitions.h
--- a/Engine/Source/Runtime/RHI/Public/RHIDefinitions.h
+++ b/Engine/Source/Runtime/RHI/Public/RHIDefinitions.h
@@ -673,8 +673,12 @@ inline bool RHINeedsToSwitchVerticalAxis(EShaderPlatform Platform)
 
 inline bool RHISupportsInstancing(const EShaderPlatform Platform)
 {
-	//@todo-rco: Add Metal support
-	return IsFeatureLevelSupported(Platform, ERHIFeatureLevel::SM4);// || (Platform == SP_METAL);
+	return IsFeatureLevelSupported(Platform, ERHIFeatureLevel::ES3_1);
+}
+
+inline bool RHISupportsInstancing(ERHIFeatureLevel::Type FeatureLevel)
+{
+	return FeatureLevel >= ERHIFeatureLevel::ES3_1;
 }
 
 inline bool RHISupportsSeparateMSAAAndResolveTextures(const EShaderPlatform Platform)
