- Metal MRT basic support (filling out RHI, disabling features via CVars, METAL_MRT_PROFILE in shaders, etc)

diff --git a/Engine/Config/BaseDeviceProfiles.ini b/Engine/Config/BaseDeviceProfiles.ini
--- a/Engine/Config/BaseDeviceProfiles.ini
+++ b/Engine/Config/BaseDeviceProfiles.ini
@@ -47,6 +47,19 @@ BaseProfileName=
 +CVars=r.DepthOfFieldQuality=0
 +CVars=r.LightShaftQuality=0
 +CVars=r.RefractionQuality=0
+// HZB will be slower with tiled, and not needed
++CVars=r.HZBOcclusion=0
+// Non need for depth prepass
++CVars=r.EarlyZPass=0
++CVars=r.EarlyZPassMovable=0
+// Needs geometry shader support
++CVars=r.TranslucentLightingVolume=0
+// Needs geometry shader support
++CVars=r.AllowPointLightCubemapShadows=0
+// TXAA (?) looks blurry right now
++CVars=r.PostProcessAAQuality=0
+// Not sure on this one - may need the whole D Buffer stuff going
++CVars=r.Decal.StencilSizeThreshold=-1
 
 [iPad2 DeviceProfile]
 DeviceType=IOS
diff --git a/Engine/Shaders/Definitions.usf b/Engine/Shaders/Definitions.usf
--- a/Engine/Shaders/Definitions.usf
+++ b/Engine/Shaders/Definitions.usf
@@ -96,6 +96,10 @@
 #define METAL_PROFILE									0
 #endif
 
+#ifndef METAL_MRT_PROFILE
+#define METAL_MRT_PROFILE								0
+#endif
+
 #ifndef IOS
 #define IOS												0
 #endif
diff --git a/Engine/Source/Runtime/IOS/MetalRHI/Private/MetalCommands.cpp b/Engine/Source/Runtime/IOS/MetalRHI/Private/MetalCommands.cpp
--- a/Engine/Source/Runtime/IOS/MetalRHI/Private/MetalCommands.cpp
+++ b/Engine/Source/Runtime/IOS/MetalRHI/Private/MetalCommands.cpp
@@ -6,7 +6,33 @@
 
 #include "MetalRHIPrivate.h"
 
+#include "GlobalShader.h"
+#include "OneColorShader.h"
+#include "RHICommandList.h"
+#include "RHIStaticStates.h"
+#include "ShaderParameterUtils.h"
+#include "SceneUtils.h"
+
 static const bool GUsesInvertedZ = true;
+static FGlobalBoundShaderState GClearMRTBoundShaderState[8];
+
+/** Vertex declaration for just one FVector4 position. */
+class FVector4VertexDeclaration : public FRenderResource
+{
+public:
+	FVertexDeclarationRHIRef VertexDeclarationRHI;
+	virtual void InitRHI() override
+	{
+		FVertexDeclarationElementList Elements;
+		Elements.Add(FVertexElement(0, 0, VET_Float4, 0, sizeof(FVector4)));
+		VertexDeclarationRHI = RHICreateVertexDeclaration(Elements);
+	}
+	virtual void ReleaseRHI() override
+	{
+		VertexDeclarationRHI.SafeRelease();
+	}
+};
+static TGlobalResource<FVector4VertexDeclaration> GVector4VertexDeclaration;
 
 static MTLPrimitiveType TranslatePrimitiveType(uint32 PrimitiveType)
 {
@@ -642,7 +668,119 @@ void FMetalDynamicRHI::RHIClear(bool bClearColor,const FLinearColor& Color, bool
 
 void FMetalDynamicRHI::RHIClearMRT(bool bClearColor,int32 NumClearColors,const FLinearColor* ClearColorArray,bool bClearDepth,float Depth,bool bClearStencil,uint32 Stencil, FIntRect ExcludeRect)
 {
-	// this needs to be handled in SetRenderTarget LoadAction 
+	// @todo metalmrt: This is currently only supporting clearing stencil by a shader (because shadows need to clear
+	// stencil in the middle of an encoder, so we can't use LoadAction). This code was adapted from D3D code (would
+	// be nice to share this code!), and parts removed until basically just what was needed for Stencil only is left
+	// If we need other types of clears, diff this against the D3D version to see what it's missing
+	if (bClearStencil && !bClearColor && !bClearDepth)
+	{
+		RHIPushEvent(TEXT("MetalClearStencil"));
+		// we don't support draw call clears before the RHI is initialized, reorder the code or make sure it's not a draw call clear
+		check(GIsRHIInitialized);
+
+ 		// Build new states
+ 		FBlendStateRHIParamRef BlendStateRHI = TStaticBlendState<CW_NONE>::GetRHI();
+
+		FRasterizerStateRHIParamRef RasterizerStateRHI = TStaticRasterizerState<FM_Solid, CM_None>::GetRHI();
+		float BF[4] = { 0, 0, 0, 0 };
+
+		const FDepthStencilStateRHIParamRef DepthStencilStateRHI = 
+			(bClearDepth && bClearStencil)
+				? TStaticDepthStencilState<
+					true, CF_Always,
+					true,CF_Always,SO_Replace,SO_Replace,SO_Replace,
+					false,CF_Always,SO_Replace,SO_Replace,SO_Replace,
+					0xff,0xff
+					>::GetRHI()
+			: bClearDepth
+				? TStaticDepthStencilState<true, CF_Always>::GetRHI()
+			: bClearStencil
+				? TStaticDepthStencilState<
+					false, CF_Always,
+					true,CF_Always,SO_Replace,SO_Replace,SO_Replace,
+					true,CF_Always,SO_Replace,SO_Replace,SO_Replace,
+					0xff,0xff
+					>::GetRHI()
+			:     TStaticDepthStencilState<false, CF_Always>::GetRHI();
+
+		RHISetBlendState(BlendStateRHI, FLinearColor::Transparent);
+		RHISetDepthStencilState(DepthStencilStateRHI, Stencil);
+		RHISetRasterizerState(RasterizerStateRHI);
+
+		// Set the new shaders
+		auto ShaderMap = GetGlobalShaderMap(GMaxRHIFeatureLevel);
+		TShaderMapRef<TOneColorVS<true> > VertexShader(ShaderMap);
+
+		FOneColorPS* PixelShader = NULL;
+
+		TShaderMapRef<TOneColorPixelShaderMRT<1> > MRTPixelShader(ShaderMap);
+		PixelShader = *MRTPixelShader;
+
+		{
+			FRHICommandList_RecursiveHazardous RHICmdList;
+			SetGlobalBoundShaderState(RHICmdList, GMaxRHIFeatureLevel, GClearMRTBoundShaderState[0], GVector4VertexDeclaration.VertexDeclarationRHI, *VertexShader, PixelShader);
+			FLinearColor ShaderClearColors[MaxSimultaneousRenderTargets];
+			FMemory::MemZero(ShaderClearColors);
+
+			for (int32 i = 0; i < NumClearColors; i++)
+			{
+				ShaderClearColors[i] = ClearColorArray[i];
+			}
+
+			SetShaderValueArray(RHICmdList, PixelShader->GetPixelShader(), PixelShader->ColorParameter, ShaderClearColors, NumClearColors);
+
+			{
+				// Draw a fullscreen quad
+				if (ExcludeRect.Width() > 0 && ExcludeRect.Height() > 0)
+				{
+					// with a hole in it (optimization in case the hardware has non constant clear performance)
+					FVector4 OuterVertices[4];
+					OuterVertices[0].Set(-1.0f, 1.0f, Depth, 1.0f);
+					OuterVertices[1].Set(1.0f, 1.0f, Depth, 1.0f);
+					OuterVertices[2].Set(1.0f, -1.0f, Depth, 1.0f);
+					OuterVertices[3].Set(-1.0f, -1.0f, Depth, 1.0f);
+
+					UE_LOG(LogMetal, Fatal, TEXT("ExcludeRect not supported yet in Metal Clear"));
+					float InvViewWidth = 1.0f / 100;//Viewport.Width;
+					float InvViewHeight = 1.0f / 100;//Viewport.Height;
+					FVector4 FractionRect = FVector4(ExcludeRect.Min.X * InvViewWidth, ExcludeRect.Min.Y * InvViewHeight, (ExcludeRect.Max.X - 1) * InvViewWidth, (ExcludeRect.Max.Y - 1) * InvViewHeight);
+
+					FVector4 InnerVertices[4];
+					InnerVertices[0].Set(FMath::Lerp(-1.0f, 1.0f, FractionRect.X), FMath::Lerp(1.0f, -1.0f, FractionRect.Y), Depth, 1.0f);
+					InnerVertices[1].Set(FMath::Lerp(-1.0f, 1.0f, FractionRect.Z), FMath::Lerp(1.0f, -1.0f, FractionRect.Y), Depth, 1.0f);
+					InnerVertices[2].Set(FMath::Lerp(-1.0f, 1.0f, FractionRect.Z), FMath::Lerp(1.0f, -1.0f, FractionRect.W), Depth, 1.0f);
+					InnerVertices[3].Set(FMath::Lerp(-1.0f, 1.0f, FractionRect.X), FMath::Lerp(1.0f, -1.0f, FractionRect.W), Depth, 1.0f);
+
+					FVector4 Vertices[10];
+					Vertices[0] = OuterVertices[0];
+					Vertices[1] = InnerVertices[0];
+					Vertices[2] = OuterVertices[1];
+					Vertices[3] = InnerVertices[1];
+					Vertices[4] = OuterVertices[2];
+					Vertices[5] = InnerVertices[2];
+					Vertices[6] = OuterVertices[3];
+					Vertices[7] = InnerVertices[3];
+					Vertices[8] = OuterVertices[0];
+					Vertices[9] = InnerVertices[0];
+
+					DrawPrimitiveUP(RHICmdList, PT_TriangleStrip, 8, Vertices, sizeof(Vertices[0]));
+				}
+				else
+				{
+					// without a hole
+					FVector4 Vertices[4];
+					Vertices[0].Set(-1.0f, 1.0f, Depth, 1.0f);
+					Vertices[1].Set(1.0f, 1.0f, Depth, 1.0f);
+					Vertices[2].Set(-1.0f, -1.0f, Depth, 1.0f);
+					Vertices[3].Set(1.0f, -1.0f, Depth, 1.0f);
+					DrawPrimitiveUP(RHICmdList, PT_TriangleStrip, 2, Vertices, sizeof(Vertices[0]));
+				}
+			}
+			// Implicit flush. Always call flush when using a command list in RHI implementations before doing anything else. This is super hazardous.
+		}
+
+		RHIPopEvent();
+	}
 }
 
 
diff --git a/Engine/Source/Runtime/IOS/MetalRHI/Private/MetalRHI.cpp b/Engine/Source/Runtime/IOS/MetalRHI/Private/MetalRHI.cpp
--- a/Engine/Source/Runtime/IOS/MetalRHI/Private/MetalRHI.cpp
+++ b/Engine/Source/Runtime/IOS/MetalRHI/Private/MetalRHI.cpp
@@ -91,7 +91,7 @@ FMetalDynamicRHI::FMetalDynamicRHI()
 	GShaderPlatformForFeatureLevel[ERHIFeatureLevel::SM4] = (GMaxRHIShaderPlatform == SP_METAL_MRT) ? SP_METAL_MRT : SP_NumPlatforms;
 	GShaderPlatformForFeatureLevel[ERHIFeatureLevel::SM5] = SP_NumPlatforms;
 
-	// @todo metal mrt:
+	// we cannot render to a volume texture without geometry shader support
 	GSupportsVolumeTextureRendering = false;
 
 // 	GDrawUPVertexCheckCount = MAX_uint16;
diff --git a/Engine/Source/Runtime/IOS/MetalRHI/Private/MetalState.cpp b/Engine/Source/Runtime/IOS/MetalRHI/Private/MetalState.cpp
--- a/Engine/Source/Runtime/IOS/MetalRHI/Private/MetalState.cpp
+++ b/Engine/Source/Runtime/IOS/MetalRHI/Private/MetalState.cpp
@@ -195,26 +195,28 @@ FMetalDepthStencilState::FMetalDepthStencilState(const FDepthStencilStateInitial
 		MTLStencilDescriptor* Stencil = [[MTLStencilDescriptor alloc] init];
 		Stencil.stencilCompareFunction = TranslateCompareFunction(Initializer.FrontFaceStencilTest);
 		Stencil.stencilFailureOperation = TranslateStencilOp(Initializer.FrontFaceStencilFailStencilOp);
-		Stencil.depthFailureOperation = TranslateStencilOp(Initializer.FrontFaceStencilFailStencilOp);
+		Stencil.depthFailureOperation = TranslateStencilOp(Initializer.FrontFaceDepthFailStencilOp);
 		Stencil.depthStencilPassOperation = TranslateStencilOp(Initializer.FrontFacePassStencilOp);
 		Stencil.readMask = Initializer.StencilReadMask;
 		Stencil.writeMask = Initializer.StencilWriteMask;
-		Desc.frontFaceStencil = Stencil;
+		// @todo metal: We had to swap, but I can't remember why these are backwards.
+		// Should the culling stuff be flipped instead? 
+		Desc.backFaceStencil = Stencil;
 		
 		[Stencil release];
 	}
 	
-	if (Initializer.bEnableFrontFaceStencil)
+	if (Initializer.bEnableBackFaceStencil)
 	{
 		// set up back face stencil operations
 		MTLStencilDescriptor* Stencil = [[MTLStencilDescriptor alloc] init];
 		Stencil.stencilCompareFunction = TranslateCompareFunction(Initializer.BackFaceStencilTest);
 		Stencil.stencilFailureOperation = TranslateStencilOp(Initializer.BackFaceStencilFailStencilOp);
-		Stencil.depthFailureOperation = TranslateStencilOp(Initializer.BackFaceStencilFailStencilOp);
+		Stencil.depthFailureOperation = TranslateStencilOp(Initializer.BackFaceDepthFailStencilOp);
 		Stencil.depthStencilPassOperation = TranslateStencilOp(Initializer.BackFacePassStencilOp);
 		Stencil.readMask = Initializer.StencilReadMask;
 		Stencil.writeMask = Initializer.StencilWriteMask;
-		Desc.backFaceStencil = Stencil;
+		Desc.frontFaceStencil = Stencil;
 		
 		[Stencil release];
 	}
diff --git a/Engine/Source/Runtime/IOS/MetalRHI/Private/MetalTexture.cpp b/Engine/Source/Runtime/IOS/MetalRHI/Private/MetalTexture.cpp
--- a/Engine/Source/Runtime/IOS/MetalRHI/Private/MetalTexture.cpp
+++ b/Engine/Source/Runtime/IOS/MetalRHI/Private/MetalTexture.cpp
@@ -100,6 +100,8 @@ FMetalSurface::FMetalSurface(ERHIResourceType ResourceType, EPixelFormat Format,
 	// set the key
 	FormatKey = *Key;
 
+	FMemory::Memzero(LockedMemory, sizeof(LockedMemory));
+
 
 	// the special back buffer surface will be updated in FMetalManager::BeginFrame - no need to set the texture here
 	if (Flags & TexCreate_Presentable)
@@ -121,12 +123,13 @@ FMetalSurface::FMetalSurface(ERHIResourceType ResourceType, EPixelFormat Format,
 	{
 		Desc = [MTLTextureDescriptor new];
 		Desc.textureType = MTLTextureType3D;
-		Desc.height = SizeX;
-		Desc.width = SizeY;
+		Desc.width = SizeX;
+		Desc.height = SizeY;
 		Desc.depth = SizeZ;
 		Desc.pixelFormat = MTLFormat;
 		Desc.arrayLength = 1;
 		Desc.mipmapLevelCount = 1;
+		Desc.sampleCount = 1;
 	}
 	else
 	{
@@ -156,16 +159,17 @@ FMetalSurface::FMetalSurface(ERHIResourceType ResourceType, EPixelFormat Format,
 	if (BulkData)
 	{
 		UE_LOG(LogIOS, Display, TEXT("Got a bulk data texture, with %d mips"), NumMips);
-		// copy over bulk data
-		// Lock
-		// Memcpy
-		// Unlock
+		checkf(NumMips == 1&& ArraySize == 1, TEXT("Only handling bulk data with 1 mip and 1 array length"));
+		uint32 Stride;
+
+		// lock, copy, unlock
+		void* LockedData = Lock(0, 0, RLM_WriteOnly, Stride);
+		FMemory::Memcpy(LockedData, BulkData->GetResourceBulkData(), BulkData->GetResourceBulkDataSize());
+		Unlock(0, 0);
 
 		// bulk data can be unloaded now
 		BulkData->Discard();
 	}
-	
-	FMemory::Memzero(LockedMemory, sizeof(LockedMemory));
 
 	// calculate size of the texture
 	TotalTextureSize = GetMemorySize();
@@ -291,7 +295,7 @@ FMetalSurface::~FMetalSurface()
 void* FMetalSurface::Lock(uint32 MipIndex, uint32 ArrayIndex, EResourceLockMode LockMode, uint32& DestStride)
 {
 	// get size and stride
-	const uint32 MipBytes = GetMipSize(MipIndex, &DestStride);
+	const uint32 MipBytes = GetMipSize(MipIndex, &DestStride, false);
 	
 	// allocate some temporary memory
 	check(LockedMemory[MipIndex] == NULL);
@@ -307,25 +311,32 @@ void* FMetalSurface::Lock(uint32 MipIndex, uint32 ArrayIndex, EResourceLockMode
 void FMetalSurface::Unlock(uint32 MipIndex, uint32 ArrayIndex)
 {
 	uint32 Stride;
-	uint32 MipBytes = GetMipSize(MipIndex, &Stride);
+	uint32 BytesPerImage = GetMipSize(MipIndex, &Stride, true);
 	if (PixelFormat == PF_PVRTC2 || PixelFormat == PF_PVRTC4)
 	{
 		// compressed textures want zero here for whatever reason
-		MipBytes = 0;
+		BytesPerImage = 0;
 		Stride = 0;
 	}
 
-	checkf(SizeZ <= 1 || SizeZ == 6, TEXT("3D textures are not supported yet (SizeZ = %d"), SizeZ);
-
-	// upload the texture to the texture slice
-	MTLRegion Region = MTLRegionMake2D(0, 0, FMath::Max<uint32>(SizeX>>MipIndex, 1), FMath::Max<uint32>(SizeY>>MipIndex, 1));
-	[Texture replaceRegion:Region mipmapLevel:MipIndex slice:ArrayIndex withBytes:LockedMemory[MipIndex] bytesPerRow:Stride bytesPerImage:MipBytes];
+	if (SizeZ <= 1 || bIsCubemap)
+	{
+		// upload the texture to the texture slice
+		MTLRegion Region = MTLRegionMake2D(0, 0, FMath::Max<uint32>(SizeX >> MipIndex, 1), FMath::Max<uint32>(SizeY >> MipIndex, 1));
+		[Texture replaceRegion : Region mipmapLevel : MipIndex slice : ArrayIndex withBytes : LockedMemory[MipIndex] bytesPerRow:Stride bytesPerImage:0];
+	}
+	else
+	{
+		// upload the texture to the texture slice
+		MTLRegion Region = MTLRegionMake3D(0, 0, 0, FMath::Max<uint32>(SizeX >> MipIndex, 1), FMath::Max<uint32>(SizeY >> MipIndex, 1), FMath::Max<uint32>(SizeZ >> MipIndex, 1));
+		[Texture replaceRegion:Region mipmapLevel:MipIndex slice:ArrayIndex withBytes:LockedMemory[MipIndex] bytesPerRow:Stride bytesPerImage:BytesPerImage];
+	}
 	
 	FMemory::Free(LockedMemory[MipIndex]);
 	LockedMemory[MipIndex] = NULL;
 }
 
-uint32 FMetalSurface::GetMipSize(uint32 MipIndex, uint32* Stride)
+uint32 FMetalSurface::GetMipSize(uint32 MipIndex, uint32* Stride, bool bSingleLayer)
 {
 	// Calculate the dimensions of the mip-map.
 	const uint32 BlockSizeX = GPixelFormats[PixelFormat].BlockSizeX;
@@ -333,7 +344,7 @@ uint32 FMetalSurface::GetMipSize(uint32 MipIndex, uint32* Stride)
 	const uint32 BlockBytes = GPixelFormats[PixelFormat].BlockBytes;
 	const uint32 MipSizeX = FMath::Max(SizeX >> MipIndex, BlockSizeX);
 	const uint32 MipSizeY = FMath::Max(SizeY >> MipIndex, BlockSizeY);
-	const uint32 MipSizeZ = FMath::Max(SizeZ >> MipIndex, 1u);
+	const uint32 MipSizeZ = bSingleLayer ? 1 : FMath::Max(SizeZ >> MipIndex, 1u);
 	uint32 NumBlocksX = (MipSizeX + BlockSizeX - 1) / BlockSizeX;
 	uint32 NumBlocksY = (MipSizeY + BlockSizeY - 1) / BlockSizeY;
 	if (PixelFormat == PF_PVRTC2 || PixelFormat == PF_PVRTC4)
@@ -369,7 +380,7 @@ uint32 FMetalSurface::GetMemorySize()
 	uint32 TotalSize = 0;
 	for (uint32 MipIndex = 0; MipIndex < [Texture mipmapLevelCount]; MipIndex++)
 	{
-		TotalSize += GetMipSize(MipIndex, NULL);
+		TotalSize += GetMipSize(MipIndex, NULL, false);
 	}
 
 	return TotalSize;
diff --git a/Engine/Source/Runtime/IOS/MetalRHI/Public/MetalResources.h b/Engine/Source/Runtime/IOS/MetalRHI/Public/MetalResources.h
--- a/Engine/Source/Runtime/IOS/MetalRHI/Public/MetalResources.h
+++ b/Engine/Source/Runtime/IOS/MetalRHI/Public/MetalResources.h
@@ -176,7 +176,7 @@ public:
 	/**
 	 * Returns how much memory a single mip uses, and optionally returns the stride
 	 */
-	uint32 GetMipSize(uint32 MipIndex, uint32* Stride);
+	uint32 GetMipSize(uint32 MipIndex, uint32* Stride, bool bSingleLayer);
 
 	/**
 	 * Returns how much memory is used by the surface
