#UE4 - Metal compute shaders runtime (initial version)

diff --git a/Engine/Source/Runtime/IOS/MetalRHI/Private/MetalCommands.cpp b/Engine/Source/Runtime/IOS/MetalRHI/Private/MetalCommands.cpp
--- a/Engine/Source/Runtime/IOS/MetalRHI/Private/MetalCommands.cpp
+++ b/Engine/Source/Runtime/IOS/MetalRHI/Private/MetalCommands.cpp
@@ -70,12 +70,13 @@ void FMetalDynamicRHI::RHISetRasterizerState(FRasterizerStateRHIParamRef NewStat
 
 void FMetalDynamicRHI::RHISetComputeShader(FComputeShaderRHIParamRef ComputeShaderRHI)
 {
-	NOT_SUPPORTED("RHISetComputeShader");
+	DYNAMIC_CAST_METALRESOURCE(ComputeShader, ComputeShader);
+	FMetalManager::Get()->SetComputeShader(ComputeShader);
 }
 
-void FMetalDynamicRHI::RHIDispatchComputeShader(uint32 ThreadGroupCountX, uint32 ThreadGroupCountY, uint32 ThreadGroupCountZ) 
+void FMetalDynamicRHI::RHIDispatchComputeShader(uint32 ThreadGroupCountX, uint32 ThreadGroupCountY, uint32 ThreadGroupCountZ)
 { 
-	NOT_SUPPORTED("RHIDispatchComputeShader");
+	FMetalManager::Get()->Dispatch(ThreadGroupCountX, ThreadGroupCountY, ThreadGroupCountZ);
 }
 
 void FMetalDynamicRHI::RHIDispatchIndirectComputeShader(FVertexBufferRHIParamRef ArgumentBufferRHI, uint32 ArgumentOffset) 
@@ -181,7 +182,15 @@ void FMetalDynamicRHI::RHISetShaderTexture(FPixelShaderRHIParamRef PixelShader,
 
 void FMetalDynamicRHI::RHISetShaderTexture(FComputeShaderRHIParamRef ComputeShader, uint32 TextureIndex, FTextureRHIParamRef NewTextureRHI)
 {
-	NOT_SUPPORTED("RHISetShaderTexture-Compute");
+	FMetalSurface* Surface = GetMetalSurfaceFromRHITexture(NewTextureRHI);
+	if (Surface != nullptr)
+	{
+		[FMetalManager::GetComputeContext() setTexture:Surface->Texture atIndex : TextureIndex];
+	}
+	else
+	{
+		[FMetalManager::GetComputeContext() setTexture:nil atIndex : TextureIndex];
+	}
 }
 
 
@@ -268,7 +277,29 @@ void FMetalDynamicRHI::RHISetShaderResourceViewParameter(FPixelShaderRHIParamRef
 
 void FMetalDynamicRHI::RHISetShaderResourceViewParameter(FComputeShaderRHIParamRef ComputeShaderRHI,uint32 TextureIndex,FShaderResourceViewRHIParamRef SRVRHI)
 {
-	NOT_SUPPORTED("RHISetShaderResourceViewParameter");
+	DYNAMIC_CAST_METALRESOURCE(ShaderResourceView, SRV);
+
+	FRHITexture* Texture = SRV->SourceTexture.GetReference();
+	if (Texture)
+	{
+		FMetalSurface* Surface = GetMetalSurfaceFromRHITexture(Texture);
+		if (Surface != nullptr)
+		{
+			[FMetalManager::GetComputeContext() setTexture:Surface->Texture atIndex : TextureIndex];
+		}
+		else
+		{
+			[FMetalManager::GetComputeContext() setTexture:nil atIndex : TextureIndex];
+		}
+	}
+	else
+	{
+		FMetalVertexBuffer* VB = SRV->SourceVertexBuffer.GetReference();
+		if (VB)
+		{
+			[FMetalManager::GetComputeContext() setBuffer:VB->Buffer offset : VB->Offset atIndex : TextureIndex];
+		}
+	}
 }
 
 
@@ -309,12 +340,11 @@ void FMetalDynamicRHI::RHISetShaderSampler(FPixelShaderRHIParamRef PixelShader,
 
 void FMetalDynamicRHI::RHISetShaderSampler(FComputeShaderRHIParamRef ComputeShader, uint32 SamplerIndex, FSamplerStateRHIParamRef NewStateRHI)
 {
-	DYNAMIC_CAST_METALRESOURCE(SamplerState,NewState);
+	DYNAMIC_CAST_METALRESOURCE(SamplerState, NewState);
 
-	NOT_SUPPORTED("RHISetSamplerState-Compute");
+	[FMetalManager::GetComputeContext() setSamplerState:NewState->State atIndex : SamplerIndex];
 }
 
-
 void FMetalDynamicRHI::RHISetShaderParameter(FVertexShaderRHIParamRef VertexShaderRHI, uint32 BufferIndex, uint32 BaseIndex, uint32 NumBytes, const void* NewValue)
 {
 	FMetalManager::Get()->GetShaderParameters(CrossCompiler::SHADER_STAGE_VERTEX).Set(BufferIndex, BaseIndex, NumBytes, NewValue);
@@ -323,7 +353,6 @@ void FMetalDynamicRHI::RHISetShaderParameter(FVertexShaderRHIParamRef VertexShad
 void FMetalDynamicRHI::RHISetShaderParameter(FHullShaderRHIParamRef HullShaderRHI, uint32 BufferIndex, uint32 BaseIndex, uint32 NumBytes, const void* NewValue)
 {
 	NOT_SUPPORTED("RHISetShaderParameter-Hull");
-
 }
 
 void FMetalDynamicRHI::RHISetShaderParameter(FPixelShaderRHIParamRef PixelShaderRHI, uint32 BufferIndex, uint32 BaseIndex, uint32 NumBytes, const void* NewValue)
@@ -334,18 +363,16 @@ void FMetalDynamicRHI::RHISetShaderParameter(FPixelShaderRHIParamRef PixelShader
 void FMetalDynamicRHI::RHISetShaderParameter(FDomainShaderRHIParamRef DomainShaderRHI, uint32 BufferIndex, uint32 BaseIndex, uint32 NumBytes, const void* NewValue)
 {
 	NOT_SUPPORTED("RHISetShaderParameter-Domain");
-
 }
 
 void FMetalDynamicRHI::RHISetShaderParameter(FGeometryShaderRHIParamRef GeometryShaderRHI, uint32 BufferIndex, uint32 BaseIndex, uint32 NumBytes, const void* NewValue)
 {
 	NOT_SUPPORTED("RHISetShaderParameter-Geometry");
-
 }
 
 void FMetalDynamicRHI::RHISetShaderParameter(FComputeShaderRHIParamRef ComputeShaderRHI,uint32 BufferIndex, uint32 BaseIndex, uint32 NumBytes, const void* NewValue)
 {
-	NOT_SUPPORTED("RHISetShaderParameter-Compute");
+	FMetalManager::Get()->GetShaderParameters(CrossCompiler::SHADER_STAGE_COMPUTE).Set(BufferIndex, BaseIndex, NumBytes, NewValue);
 }
 
 void FMetalDynamicRHI::RHISetShaderUniformBuffer(FVertexShaderRHIParamRef VertexShaderRHI, uint32 BufferIndex, FUniformBufferRHIParamRef BufferRHI)
@@ -393,9 +420,19 @@ void FMetalDynamicRHI::RHISetShaderUniformBuffer(FPixelShaderRHIParamRef PixelSh
 	}
 }
 
-void FMetalDynamicRHI::RHISetShaderUniformBuffer(FComputeShaderRHIParamRef ComputeShader, uint32 BufferIndex, FUniformBufferRHIParamRef BufferRHI)
+void FMetalDynamicRHI::RHISetShaderUniformBuffer(FComputeShaderRHIParamRef ComputeShaderRHI, uint32 BufferIndex, FUniformBufferRHIParamRef BufferRHI)
 {
-	NOT_SUPPORTED("RHISetShaderUniformBuffer-Compute");
+	DYNAMIC_CAST_METALRESOURCE(ComputeShader, ComputeShader);
+	ComputeShader->BoundUniformBuffers[BufferIndex] = BufferRHI;
+	ComputeShader->DirtyUniformBuffers |= 1 << BufferIndex;
+
+	auto& Bindings = ComputeShader->Bindings;
+	check(BufferIndex < Bindings.NumUniformBuffers);
+	if (Bindings.bHasRegularUniformBuffers)
+	{
+		auto* UB = (FMetalUniformBuffer*)BufferRHI;
+		[FMetalManager::GetComputeContext() setBuffer:UB->Buffer offset : UB->Offset atIndex : BufferIndex];
+	}
 }
 
 
diff --git a/Engine/Source/Runtime/IOS/MetalRHI/Private/MetalManager.cpp b/Engine/Source/Runtime/IOS/MetalRHI/Private/MetalManager.cpp
--- a/Engine/Source/Runtime/IOS/MetalRHI/Private/MetalManager.cpp
+++ b/Engine/Source/Runtime/IOS/MetalRHI/Private/MetalManager.cpp
@@ -150,6 +150,11 @@ id<MTLRenderCommandEncoder> FMetalManager::GetContext()
 	return FMetalManager::Get()->CurrentContext;
 }
 
+id<MTLComputeCommandEncoder> FMetalManager::GetComputeContext()
+{
+	return FMetalManager::Get()->CurrentComputeContext;
+}
+
 void FMetalManager::ReleaseObject(id Object)
 {
 	FMetalManager::Get()->DelayedFreeLists[FMetalManager::Get()->WhichFreeList].Add(Object);
@@ -161,6 +166,7 @@ FMetalManager::FMetalManager()
 	, CurrentCommandBuffer(nil)
 	, CurrentDrawable(nil)
 	, CurrentContext(nil)
+	, CurrentComputeContext(nil)
 	, CurrentNumRenderTargets(0)
 	, PreviousNumRenderTargets(0)
 	, CurrentDepthRenderTexture(nil)
@@ -204,15 +210,15 @@ FMetalManager::FMetalManager()
 	CommandBufferSemaphore = dispatch_semaphore_create(FParse::Param(FCommandLine::Get(),TEXT("gpulockstep")) ? 1 : 3);
 
 	AutoReleasePoolTLSSlot = FPlatformTLS::AllocTlsSlot();
-    
-    // Hook into the ios framepacer, if it's enabled for this platform.
-    FrameReadyEvent = NULL;
-    if( FIOSPlatformRHIFramePacer::IsEnabled() )
-    {
-        FrameReadyEvent = FPlatformProcess::CreateSynchEvent();
-        FIOSPlatformRHIFramePacer::InitWithEvent( FrameReadyEvent );
-    }
-    
+	
+	// Hook into the ios framepacer, if it's enabled for this platform.
+	FrameReadyEvent = NULL;
+	if( FIOSPlatformRHIFramePacer::IsEnabled() )
+	{
+		FrameReadyEvent = FPlatformProcess::CreateSynchEvent();
+		FIOSPlatformRHIFramePacer::InitWithEvent( FrameReadyEvent );
+	}
+	
 	InitFrame();
 }
 
@@ -256,7 +262,6 @@ void FMetalManager::BeginFrame()
 {
 }
 
-
 void FMetalManager::InitFrame()
 {
 	// reset the event index for this frame
@@ -351,6 +356,33 @@ void FMetalManager::SubmitCommandBufferAndWait()
 	CreateCurrentCommandBuffer(true);
 }
 
+void FMetalManager::SubmitComputeCommandBufferAndWait()
+{
+	[CurrentCommandBuffer addCompletedHandler : ^ (id <MTLCommandBuffer> Buffer)
+	{
+		dispatch_semaphore_signal(CommandBufferSemaphore);
+	}];
+
+	// commit the compute context to the commandBuffer
+	[CurrentComputeContext endEncoding];
+	[CurrentComputeContext release];
+	CurrentComputeContext = nil;
+
+	// kick the whole buffer
+	// Commit to hand the commandbuffer off to the gpu
+	[CurrentCommandBuffer commit];
+
+	// wait for the gpu to finish executing our commands.
+	[CurrentCommandBuffer waitUntilCompleted];
+
+	//once a commandbuffer is commited it can't be added to again.
+	UNTRACK_OBJECT(CurrentCommandBuffer);
+	[CurrentCommandBuffer release];
+
+	// create a new command buffer.
+	CreateCurrentCommandBuffer(true);
+}
+
 void FMetalManager::EndFrame(bool bPresent)
 {
 //	NSLog(@"There were %d draw calls for final RT in frame %lld", NumDrawCalls, GFrameCounter);
@@ -364,12 +396,12 @@ void FMetalManager::EndFrame(bool bPresent)
 	 {
 		dispatch_semaphore_signal(CommandBufferSemaphore);
 	 }];
-    
-    // We may be limiting our framerate to the display link
-    if( FrameReadyEvent != nullptr )
-    {
-        FrameReadyEvent->Wait();
-    }
+	
+	// We may be limiting our framerate to the display link
+	if( FrameReadyEvent != nullptr )
+	{
+		FrameReadyEvent->Wait();
+	}
 
 	// Commit before waiting to avoid leaving the gpu idle
 	[CurrentCommandBuffer commit];
@@ -515,15 +547,15 @@ bool FMetalManager::NeedsToSetRenderTarget(const FRHISetRenderTargetsInfo& Rende
 			const FRHIRenderTargetView& RenderTargetView = RenderTargetsInfo.ColorRenderTarget[RenderTargetIndex];
 			const FRHIRenderTargetView& PreviousRenderTargetView = PreviousRenderTargetsInfo.ColorRenderTarget[RenderTargetIndex];
 
-            // handle simple case of switching textures or mip/slice
-            if (RenderTargetView.Texture != PreviousRenderTargetView.Texture ||
-                RenderTargetView.MipIndex != PreviousRenderTargetView.MipIndex ||
-                RenderTargetView.ArraySliceIndex != PreviousRenderTargetView.ArraySliceIndex)
-            {
-                bAllChecksPassed = false;
-                break;
-            }
-            
+			// handle simple case of switching textures or mip/slice
+			if (RenderTargetView.Texture != PreviousRenderTargetView.Texture ||
+				RenderTargetView.MipIndex != PreviousRenderTargetView.MipIndex ||
+				RenderTargetView.ArraySliceIndex != PreviousRenderTargetView.ArraySliceIndex)
+			{
+				bAllChecksPassed = false;
+				break;
+			}
+			
 			// it's non-trivial when we need to switch based on load/store action:
 			// LoadAction - it only matters what we are switching to in the new one
 			//    If we switch to Load, no need to switch as we can re-use what we already have
@@ -772,21 +804,21 @@ uint32 FMetalManager::AllocateFromQueryBuffer()
 FORCEINLINE void SetResource(uint32 ShaderStage, uint32 BindIndex, FMetalSurface* RESTRICT Surface)
 {
 //	check(Surface->Texture != nil);
-    if (Surface != nullptr)
-    {
-        if (ShaderStage == CrossCompiler::SHADER_STAGE_PIXEL)
-        {
-            [FMetalManager::GetContext() setFragmentTexture:Surface->Texture atIndex:BindIndex];
-        }
-        else
-        {
-            [FMetalManager::GetContext() setVertexTexture:Surface->Texture atIndex : BindIndex];
-        }
-    }
-    else
-    {
-        [FMetalManager::GetContext() setVertexTexture:nil atIndex : BindIndex];
-    }
+	if (Surface != nullptr)
+	{
+		if (ShaderStage == CrossCompiler::SHADER_STAGE_PIXEL)
+		{
+			[FMetalManager::GetContext() setFragmentTexture:Surface->Texture atIndex:BindIndex];
+		}
+		else
+		{
+			[FMetalManager::GetContext() setVertexTexture:Surface->Texture atIndex : BindIndex];
+		}
+	}
+	else
+	{
+		[FMetalManager::GetContext() setVertexTexture:nil atIndex : BindIndex];
+	}
 }
 
 FORCEINLINE void SetResource(uint32 ShaderStage, uint32 BindIndex, FMetalSamplerState* RESTRICT SamplerState)
@@ -916,4 +948,16 @@ void FMetalManager::SetRasterizerState(const FRasterizerStateInitializerRHI& Sta
 #endif
 
 	bFirstRasterizerState = false;
-}
\ No newline at end of file
+}
+
+void FMetalManager::Dispatch(uint32 ThreadGroupCountX, uint32 ThreadGroupCountY, uint32 ThreadGroupCountZ)
+{
+	check(CurrentComputeShader);
+	auto* ComputeShader = (FMetalComputeShader*)CurrentComputeShader;
+	MTLSize ThreadgroupCounts = MTLSizeMake(ComputeShader->NumThreadsX, ComputeShader->NumThreadsY, ComputeShader->NumThreadsZ);
+	MTLSize Threadgroups = MTLSizeMake(ThreadGroupCountX, ThreadGroupCountY, ThreadGroupCountZ);
+	//@todo-rco: setThreadgroupMemoryLength?
+	[CurrentComputeContext dispatchThreadgroups:Threadgroups threadsPerThreadgroup:ThreadgroupCounts];
+	//@todo: Here?
+	SubmitComputeCommandBufferAndWait();
+}
diff --git a/Engine/Source/Runtime/IOS/MetalRHI/Private/MetalManager.h b/Engine/Source/Runtime/IOS/MetalRHI/Private/MetalManager.h
--- a/Engine/Source/Runtime/IOS/MetalRHI/Private/MetalManager.h
+++ b/Engine/Source/Runtime/IOS/MetalRHI/Private/MetalManager.h
@@ -42,7 +42,7 @@ struct FPipelineShadow
 	MTLRenderPipelineColorAttachmentDescriptor* RenderTargets[MaxMetalRenderTargets];
 	MTLPixelFormat DepthTargetFormat;
 	MTLPixelFormat StencilTargetFormat;
-    uint32 SampleCount;
+	uint32 SampleCount;
 
 	// running hash of the pipeline state
 	uint64 Hash;
@@ -70,6 +70,7 @@ public:
 	static FMetalManager* Get();
 	static id<MTLDevice> GetDevice();
 	static id<MTLRenderCommandEncoder> GetContext();
+	static id<MTLComputeCommandEncoder> GetComputeContext();
 	static void ReleaseObject(id Object);
 	
 	/** RHIBeginScene helper */
@@ -136,6 +137,7 @@ public:
 	}
 
 	void SubmitCommandBufferAndWait();
+	void SubmitComputeCommandBufferAndWait();
 
 	void SetRasterizerState(const FRasterizerStateInitializerRHI& State);
 
@@ -158,7 +160,8 @@ public:
 	void CreateAutoreleasePool();
 	void DrainAutoreleasePool();
 
-
+	void SetComputeShader(FComputeShaderRHIParamRef InComputeShader) { CurrentComputeShader = InComputeShader; }
+	void Dispatch(uint32 ThreadGroupCountX, uint32 ThreadGroupCountY, uint32 ThreadGroupCountZ);
 
 protected:
 	FMetalManager();
@@ -189,7 +192,8 @@ protected:
 	dispatch_semaphore_t CommandBufferSemaphore;
 
 	id<MTLRenderCommandEncoder> CurrentContext;
-	
+	id<MTLComputeCommandEncoder> CurrentComputeContext;
+
 	struct FRenderTargetViewInfo
 	{
 		uint32 MipIndex;
@@ -228,11 +232,13 @@ protected:
 	
 	class FMetalShaderParameterCache*	ShaderParameters;
 
+	FComputeShaderRHIParamRef CurrentComputeShader;
+
 	// the running pipeline state descriptor object
 	FPipelineShadow Pipeline;
 
 	MTLRenderPipelineDescriptor* CurrentPipelineDescriptor;
-	// a parellel descriptor for no depth buffer (since we can't reset the depth format to invalid YET)
+	// a parallel descriptor for no depth buffer (since we can't reset the depth format to invalid YET)
 	MTLRenderPipelineDescriptor* CurrentPipelineDescriptor2;
 	
 
@@ -270,9 +276,9 @@ protected:
 	void CommitNonComputeShaderConstants();
 
 private:
-    
-    /** Event for coordinating pausing of render thread to keep inline with the ios display link. */
-    FEvent* FrameReadyEvent;
+	
+	/** Event for coordinating pausing of render thread to keep inline with the ios display link. */
+	FEvent* FrameReadyEvent;
 };
 
 // Stats
