- Removed old Metal hack

diff --git a/Engine/Source/Runtime/IOS/MetalRHI/Private/MetalVertexBuffer.cpp b/Engine/Source/Runtime/IOS/MetalRHI/Private/MetalVertexBuffer.cpp
--- a/Engine/Source/Runtime/IOS/MetalRHI/Private/MetalVertexBuffer.cpp
+++ b/Engine/Source/Runtime/IOS/MetalRHI/Private/MetalVertexBuffer.cpp
@@ -20,13 +20,6 @@ FMetalVertexBuffer::FMetalVertexBuffer(uint32 InSize, uint32 InUsage)
 	}
 	else
 	{
-		// reduce some churn in the ZeroStride buffer hack
-		if (InSize <= 16)
-		{
-			InSize = 128 * 1024;
-		}
-//		NSLog(@"Allocating %d byte vertex buffer", InSize);
-		
 		Buffer = [FMetalManager::GetDevice() newBufferWithLength:InSize options:BUFFER_CACHE_MODE];
 		Offset = 0;
 	}
