- First pass of compute shader support for Metal (currently only MetalMRT). GPU sorting isn't working, so this is not complete.

diff --git a/Engine/Source/Editor/UnrealEd/Private/ThumbnailHelpers.cpp b/Engine/Source/Editor/UnrealEd/Private/ThumbnailHelpers.cpp
--- a/Engine/Source/Editor/UnrealEd/Private/ThumbnailHelpers.cpp
+++ b/Engine/Source/Editor/UnrealEd/Private/ThumbnailHelpers.cpp
@@ -146,7 +146,7 @@ FParticleSystemThumbnailScene::FParticleSystemThumbnailScene()
 	bForceAllUsedMipsResident = false;
 	PartComponent = NULL;
 
-	ThumbnailFXSystem = FFXSystemInterface::Create(GetScene()->GetFeatureLevel());
+	ThumbnailFXSystem = FFXSystemInterface::Create(GetScene()->GetFeatureLevel(), GetScene()->GetShaderPlatform());
 	GetScene()->SetFXSystem( ThumbnailFXSystem );
 }
 
diff --git a/Engine/Source/Runtime/Engine/Private/GPUSort.cpp b/Engine/Source/Runtime/Engine/Private/GPUSort.cpp
--- a/Engine/Source/Runtime/Engine/Private/GPUSort.cpp
+++ b/Engine/Source/Runtime/Engine/Private/GPUSort.cpp
@@ -82,7 +82,7 @@ public:
 	 */
 	virtual void InitRHI() override
 	{
-		if (GetFeatureLevel() >= ERHIFeatureLevel::SM5)
+		if (RHISupportsComputeShaders(GShaderPlatformForFeatureLevel[GetFeatureLevel()]))
 		{
 			const int32 OffsetsCount = DIGIT_COUNT * MAX_GROUP_COUNT;
 			const int32 OffsetsBufferSize = OffsetsCount * sizeof(uint32);
@@ -183,7 +183,7 @@ public:
 	 */
 	virtual void InitRHI() override
 	{
-		if (GetFeatureLevel() >= ERHIFeatureLevel::SM5)
+		if (RHISupportsComputeShaders(GShaderPlatformForFeatureLevel[GetFeatureLevel()]))
 		{
 			FRHIResourceCreateInfo CreateInfo;
 			SortParametersBufferRHI = RHICreateVertexBuffer(
@@ -686,7 +686,7 @@ int32 SortGPUBuffers(FRHICommandListImmediate& RHICmdList, FGPUSortBuffers SortB
 	const bool bDebugOffsets = CVarDebugOffsets.GetValueOnRenderThread() != 0;
 	const bool bDebugSort = CVarDebugSort.GetValueOnRenderThread() != 0;
 
-	check(FeatureLevel == ERHIFeatureLevel::SM5);
+	check(RHISupportsComputeShaders(GShaderPlatformForFeatureLevel[FeatureLevel]));
 
 	SCOPED_DRAW_EVENTF(RHICmdList, SortGPU, TEXT("SortGPU_%d"), Count);
 
@@ -845,7 +845,7 @@ static bool RunGPUSortTest(FRHICommandListImmediate& RHICmdList, int32 TestSize,
 	const bool bDebugOffsets = CVarDebugOffsets.GetValueOnRenderThread() != 0;
 	const bool bDebugSort = CVarDebugSort.GetValueOnRenderThread() != 0;
 
-	if (FeatureLevel != ERHIFeatureLevel::SM5)
+	if (!RHISupportsComputeShaders(GShaderPlatformForFeatureLevel[FeatureLevel]))
 	{
 		return false;
 	}
diff --git a/Engine/Source/Runtime/Engine/Private/Particles/FXSystem.cpp b/Engine/Source/Runtime/Engine/Private/Particles/FXSystem.cpp
--- a/Engine/Source/Runtime/Engine/Private/Particles/FXSystem.cpp
+++ b/Engine/Source/Runtime/Engine/Private/Particles/FXSystem.cpp
@@ -20,9 +20,9 @@
 	External FX system interface.
 -----------------------------------------------------------------------------*/
 
-FFXSystemInterface* FFXSystemInterface::Create(ERHIFeatureLevel::Type InFeatureLevel)
+FFXSystemInterface* FFXSystemInterface::Create(ERHIFeatureLevel::Type InFeatureLevel, EShaderPlatform InShaderPlatform)
 {
-	return new FFXSystem(InFeatureLevel);
+	return new FFXSystem(InFeatureLevel, InShaderPlatform);
 }
 
 void FFXSystemInterface::Destroy( FFXSystemInterface* FXSystem )
@@ -141,9 +141,10 @@ namespace FXConsoleVariables
 	FX system.
 ------------------------------------------------------------------------------*/
 
-FFXSystem::FFXSystem(ERHIFeatureLevel::Type InFeatureLevel)
+FFXSystem::FFXSystem(ERHIFeatureLevel::Type InFeatureLevel, EShaderPlatform InShaderPlatform)
 	: ParticleSimulationResources(NULL)
 	, FeatureLevel(InFeatureLevel)
+	, ShaderPlatform(InShaderPlatform)
 #if WITH_EDITOR
 	, bSuspended(false)
 #endif // #if WITH_EDITOR
diff --git a/Engine/Source/Runtime/Engine/Private/Particles/FXSystemPrivate.h b/Engine/Source/Runtime/Engine/Private/Particles/FXSystemPrivate.h
--- a/Engine/Source/Runtime/Engine/Private/Particles/FXSystemPrivate.h
+++ b/Engine/Source/Runtime/Engine/Private/Particles/FXSystemPrivate.h
@@ -48,7 +48,7 @@ class FFXSystem : public FFXSystemInterface
 public:
 
 	/** Default constructoer. */
-	FFXSystem(ERHIFeatureLevel::Type InFeatureLevel);
+	FFXSystem(ERHIFeatureLevel::Type InFeatureLevel, EShaderPlatform InShaderPlatform);
 
 	/** Destructor. */
 	virtual ~FFXSystem();
@@ -77,6 +77,11 @@ public:
 	 */
 	ERHIFeatureLevel::Type GetFeatureLevel() const { return FeatureLevel; }
 
+	/**
+	 * Retrieve shaderplatform that this FXSystem was created for
+	 */
+	EShaderPlatform GetShaderPlatform() const { return ShaderPlatform; }
+
 	/**
 	 * Add a new GPU simulation to the system.
 	 * @param Simulation The GPU simulation to add.
@@ -177,6 +182,8 @@ private:
 	FParticleSimulationResources* ParticleSimulationResources;
 	/** Feature level of this effects system */
 	ERHIFeatureLevel::Type FeatureLevel;
+	/** Shader platform that will be rendering this effects system */
+	EShaderPlatform ShaderPlatform;
 
 #if WITH_EDITOR
 	/** true if the system has been suspended. */
diff --git a/Engine/Source/Runtime/Engine/Private/Particles/ParticleSortingGPU.cpp b/Engine/Source/Runtime/Engine/Private/Particles/ParticleSortingGPU.cpp
--- a/Engine/Source/Runtime/Engine/Private/Particles/ParticleSortingGPU.cpp
+++ b/Engine/Source/Runtime/Engine/Private/Particles/ParticleSortingGPU.cpp
@@ -49,7 +49,7 @@ public:
 
 	static bool ShouldCache( EShaderPlatform Platform )
 	{
-		return IsFeatureLevelSupported(Platform, ERHIFeatureLevel::SM5);
+		return RHISupportsComputeShaders(Platform);
 	}
 
 	static void ModifyCompilationEnvironment( EShaderPlatform Platform, FShaderCompilerEnvironment& OutEnvironment )
@@ -185,7 +185,7 @@ static int32 GenerateParticleSortKeys(
 	)
 {
 	SCOPED_DRAW_EVENT(RHICmdList, ParticleSortKeyGen);
-	check(FeatureLevel == ERHIFeatureLevel::SM5);
+	check(RHISupportsComputeShaders(GShaderPlatformForFeatureLevel[FeatureLevel]));
 
 	FParticleKeyGenParameters KeyGenParameters;
 	FParticleKeyGenUniformBufferRef KeyGenUniformBuffer;
@@ -240,7 +240,7 @@ static int32 GenerateParticleSortKeys(
  */
 void FParticleSortBuffers::InitRHI()
 {
-	if (GetFeatureLevel() >= ERHIFeatureLevel::SM5)
+	if (RHISupportsComputeShaders(GShaderPlatformForFeatureLevel[GetFeatureLevel()]))
 	{
 		for (int32 BufferIndex = 0; BufferIndex < 2; ++BufferIndex)
 		{
diff --git a/Engine/Source/Runtime/Engine/Private/World.cpp b/Engine/Source/Runtime/Engine/Private/World.cpp
--- a/Engine/Source/Runtime/Engine/Private/World.cpp
+++ b/Engine/Source/Runtime/Engine/Private/World.cpp
@@ -5313,7 +5313,7 @@ void UWorld::CreateFXSystem()
 {
 	if ( !IsRunningDedicatedServer() && !IsRunningCommandlet() )
 	{
-		FXSystem = FFXSystemInterface::Create(FeatureLevel);
+		FXSystem = FFXSystemInterface::Create(FeatureLevel, Scene->GetShaderPlatform());
 		Scene->SetFXSystem(FXSystem);
 	}
 	else
diff --git a/Engine/Source/Runtime/Engine/Public/FXSystem.h b/Engine/Source/Runtime/Engine/Public/FXSystem.h
--- a/Engine/Source/Runtime/Engine/Public/FXSystem.h
+++ b/Engine/Source/Runtime/Engine/Public/FXSystem.h
@@ -97,7 +97,7 @@ public:
 	/**
 	 * Create an effects system instance.
 	 */
-	ENGINE_API static FFXSystemInterface* Create(ERHIFeatureLevel::Type InFeatureLevel);
+	ENGINE_API static FFXSystemInterface* Create(ERHIFeatureLevel::Type InFeatureLevel, EShaderPlatform InShaderPlatform);
 
 	/**
 	 * Destroy an effects system instance.
diff --git a/Engine/Source/Runtime/IOS/MetalRHI/Private/MetalCommands.cpp b/Engine/Source/Runtime/IOS/MetalRHI/Private/MetalCommands.cpp
--- a/Engine/Source/Runtime/IOS/MetalRHI/Private/MetalCommands.cpp
+++ b/Engine/Source/Runtime/IOS/MetalRHI/Private/MetalCommands.cpp
@@ -72,6 +72,9 @@ void FMetalDynamicRHI::RHISetComputeShader(FComputeShaderRHIParamRef ComputeShad
 {
 	DYNAMIC_CAST_METALRESOURCE(ComputeShader, ComputeShader);
 	FMetalManager::Get()->SetComputeShader(ComputeShader);
+	
+	// set this compute shader pipeline as the current (this resets all state, so we need to set all resources after calling this)
+	[FMetalManager::GetComputeContext() setComputePipelineState:ComputeShader->Kernel];
 }
 
 void FMetalDynamicRHI::RHIDispatchComputeShader(uint32 ThreadGroupCountX, uint32 ThreadGroupCountY, uint32 ThreadGroupCountZ)
@@ -129,7 +132,12 @@ void FMetalDynamicRHI::RHISetBoundShaderState( FBoundShaderStateRHIParamRef Boun
 
 void FMetalDynamicRHI::RHISetUAVParameter(FComputeShaderRHIParamRef ComputeShaderRHI, uint32 UAVIndex, FUnorderedAccessViewRHIParamRef UAVRHI)
 {
-	NOT_SUPPORTED("RHISetUAVParameter");
+	DYNAMIC_CAST_METALRESOURCE(UnorderedAccessView, UAV);
+
+	if (UAV)
+	{
+		UAV->Set(UAVIndex);
+	}
 }
 
 void FMetalDynamicRHI::RHISetUAVParameter(FComputeShaderRHIParamRef ComputeShaderRHI,uint32 UAVIndex,FUnorderedAccessViewRHIParamRef UAVRHI, uint32 InitialCount)
@@ -279,25 +287,28 @@ void FMetalDynamicRHI::RHISetShaderResourceViewParameter(FComputeShaderRHIParamR
 {
 	DYNAMIC_CAST_METALRESOURCE(ShaderResourceView, SRV);
 
-	FRHITexture* Texture = SRV->SourceTexture.GetReference();
-	if (Texture)
+	if (SRV != NULL)
 	{
-		FMetalSurface* Surface = GetMetalSurfaceFromRHITexture(Texture);
-		if (Surface != nullptr)
+		FRHITexture* Texture = SRV->SourceTexture.GetReference();
+		if (Texture)
 		{
-			[FMetalManager::GetComputeContext() setTexture:Surface->Texture atIndex : TextureIndex];
+			FMetalSurface* Surface = GetMetalSurfaceFromRHITexture(Texture);
+			if (Surface != nullptr)
+			{
+				[FMetalManager::GetComputeContext() setTexture:Surface->Texture atIndex : TextureIndex];
+			}
+			else
+			{
+				[FMetalManager::GetComputeContext() setTexture:nil atIndex : TextureIndex];
+			}
 		}
 		else
 		{
-			[FMetalManager::GetComputeContext() setTexture:nil atIndex : TextureIndex];
-		}
-	}
-	else
-	{
-		FMetalVertexBuffer* VB = SRV->SourceVertexBuffer.GetReference();
-		if (VB)
-		{
-			[FMetalManager::GetComputeContext() setBuffer:VB->Buffer offset : VB->Offset atIndex : TextureIndex];
+			FMetalVertexBuffer* VB = SRV->SourceVertexBuffer.GetReference();
+			if (VB)
+			{
+				[FMetalManager::GetComputeContext() setBuffer:VB->Buffer offset : VB->Offset atIndex : TextureIndex];
+			}
 		}
 	}
 }
@@ -459,6 +470,8 @@ void FMetalDynamicRHI::RHISetRenderTargets(uint32 NumSimultaneousRenderTargets,
 	FRHIDepthRenderTargetView DepthView(NewDepthStencilTargetRHI, ERenderTargetLoadAction::EClear, ERenderTargetStoreAction::ENoAction);
 	FRHISetRenderTargetsInfo Info(NumSimultaneousRenderTargets, NewRenderTargets, DepthView);
 	RHISetRenderTargetsAndClear(Info);
+	
+	checkf(NumUAVs == 0, TEXT("Calling SetRenderTargets with UAVs is not supported in Metal yet"));
 }
 
 void FMetalDynamicRHI::RHIDiscardRenderTargets(bool Depth, bool Stencil, uint32 ColorBitMask)
diff --git a/Engine/Source/Runtime/IOS/MetalRHI/Private/MetalManager.cpp b/Engine/Source/Runtime/IOS/MetalRHI/Private/MetalManager.cpp
--- a/Engine/Source/Runtime/IOS/MetalRHI/Private/MetalManager.cpp
+++ b/Engine/Source/Runtime/IOS/MetalRHI/Private/MetalManager.cpp
@@ -152,7 +152,9 @@ id<MTLRenderCommandEncoder> FMetalManager::GetContext()
 
 id<MTLComputeCommandEncoder> FMetalManager::GetComputeContext()
 {
-	return FMetalManager::Get()->CurrentComputeContext;
+	id<MTLComputeCommandEncoder> Encoder = FMetalManager::Get()->CurrentComputeContext;
+	checkf(Encoder != nil, TEXT("Attempted to use GetComputeContext before calling FMetalManager::ConditionalSwitchToCompute() to create the encoder"));
+	return Encoder;
 }
 
 void FMetalManager::ReleaseObject(id Object)
@@ -358,38 +360,24 @@ void FMetalManager::SubmitCommandBufferAndWait()
 
 void FMetalManager::SubmitComputeCommandBufferAndWait()
 {
-	[CurrentCommandBuffer addCompletedHandler : ^ (id <MTLCommandBuffer> Buffer)
-	{
-		dispatch_semaphore_signal(CommandBufferSemaphore);
-	}];
-
-	// commit the compute context to the commandBuffer
-	[CurrentComputeContext endEncoding];
-	[CurrentComputeContext release];
-	CurrentComputeContext = nil;
-
-	// kick the whole buffer
-	// Commit to hand the commandbuffer off to the gpu
-	[CurrentCommandBuffer commit];
-
-	// wait for the gpu to finish executing our commands.
-	[CurrentCommandBuffer waitUntilCompleted];
-
-	//once a commandbuffer is commited it can't be added to again.
-	UNTRACK_OBJECT(CurrentCommandBuffer);
-	[CurrentCommandBuffer release];
-
-	// create a new command buffer.
-	CreateCurrentCommandBuffer(true);
 }
 
 void FMetalManager::EndFrame(bool bPresent)
 {
 //	NSLog(@"There were %d draw calls for final RT in frame %lld", NumDrawCalls, GFrameCounter);
 	// commit the render context to the commandBuffer
-	[CurrentContext endEncoding];
-	[CurrentContext release];
-	CurrentContext = nil;
+	if (CurrentContext)
+	{
+		[CurrentContext endEncoding];
+		[CurrentContext release];
+		CurrentContext = nil;
+	}
+	if (CurrentComputeContext)
+	{
+		[CurrentComputeContext endEncoding];
+		[CurrentComputeContext release];
+		CurrentComputeContext = nil;
+	}
 
 	// kick the whole buffer
 	[CurrentCommandBuffer addCompletedHandler : ^ (id <MTLCommandBuffer> Buffer)
@@ -521,6 +509,8 @@ void FMetalManager::ConditionalUpdateBackBuffer(FMetalSurface& Surface)
 			// make a drawable object for this frame
 			CurrentDrawable = [[IOSAppDelegate GetDelegate].IOSView MakeDrawable];
 
+			checkf(CurrentDrawable != nil, TEXT("Failed to return the next drawable"));
+			
 			GRenderThreadIdle[ERenderThreadIdleTypes::WaitingForGPUPresent] += FPlatformTime::Cycles() - IdleStart;
 			GRenderThreadNumIdle[ERenderThreadIdleTypes::WaitingForGPUPresent]++;
 
@@ -592,6 +582,8 @@ bool FMetalManager::NeedsToSetRenderTarget(const FRHISetRenderTargetsInfo& Rende
 
 void FMetalManager::SetRenderTargetsInfo(const FRHISetRenderTargetsInfo& RenderTargetsInfo)
 {
+	ConditionalSwitchToGraphics();
+	
 	// see if our new Info matches our previous Info
 	if (NeedsToSetRenderTarget(RenderTargetsInfo) == false)
 	{
@@ -950,6 +942,56 @@ void FMetalManager::SetRasterizerState(const FRasterizerStateInitializerRHI& Sta
 	bFirstRasterizerState = false;
 }
 
+
+void FMetalManager::ConditionalSwitchToGraphics()
+{
+	// were we in compute mode?
+	if (CurrentComputeContext != nil)
+	{
+		// stop the compute encoding and cleanup up
+		[CurrentComputeContext endEncoding];
+		[CurrentComputeContext release];
+		CurrentComputeContext = nil;
+	
+		// we can't start graphics encoding until a new SetRenderTargetsInfo is called because it needs the whole render pass
+		// we could cache the render pass if we want to support going back to previous render targets
+		// we will catch it by the fact that CurrentContext will be nil until we call SetRenderTargetsInfo
+	}
+}
+
+void FMetalManager::ConditionalSwitchToCompute()
+{
+	// if we were in graphics mode, stop the encoding and start compute
+	if (CurrentContext != nil)
+	{
+		// stop encoding graphics and clean up
+		[CurrentContext endEncoding];
+		[CurrentContext release];
+		CurrentContext = nil;
+		
+		// make sure that the next SetRenderTargetInfos will definitely set a new RT (if we SetRT to X, switched to compute, then back to X,
+		// the code would think it was already set and skip it, so this will make it so the next SetRT is 'different')
+		PreviousRenderTargetsInfo.NumColorRenderTargets = 0;
+
+		// start encoding for compute
+		CurrentComputeContext = [CurrentCommandBuffer computeCommandEncoder];
+		[CurrentComputeContext retain];
+	}
+	
+}
+
+void FMetalManager::SetComputeShader(FMetalComputeShader* InComputeShader)
+{
+	// active compute mode
+	ConditionalSwitchToCompute();
+	
+	// cache this for Dispatch
+	CurrentComputeShader = InComputeShader;
+	
+	// set this compute shader pipeline as the current (this resets all state, so we need to set all resources after calling this)
+	[CurrentComputeContext setComputePipelineState:CurrentComputeShader->Kernel];
+}
+
 void FMetalManager::Dispatch(uint32 ThreadGroupCountX, uint32 ThreadGroupCountY, uint32 ThreadGroupCountZ)
 {
 	check(CurrentComputeShader);
@@ -958,6 +1000,4 @@ void FMetalManager::Dispatch(uint32 ThreadGroupCountX, uint32 ThreadGroupCountY,
 	MTLSize Threadgroups = MTLSizeMake(ThreadGroupCountX, ThreadGroupCountY, ThreadGroupCountZ);
 	//@todo-rco: setThreadgroupMemoryLength?
 	[CurrentComputeContext dispatchThreadgroups:Threadgroups threadsPerThreadgroup:ThreadgroupCounts];
-	//@todo: Here?
-	SubmitComputeCommandBufferAndWait();
 }
diff --git a/Engine/Source/Runtime/IOS/MetalRHI/Private/MetalManager.h b/Engine/Source/Runtime/IOS/MetalRHI/Private/MetalManager.h
--- a/Engine/Source/Runtime/IOS/MetalRHI/Private/MetalManager.h
+++ b/Engine/Source/Runtime/IOS/MetalRHI/Private/MetalManager.h
@@ -160,7 +160,7 @@ public:
 	void CreateAutoreleasePool();
 	void DrainAutoreleasePool();
 
-	void SetComputeShader(FComputeShaderRHIParamRef InComputeShader) { CurrentComputeShader = InComputeShader; }
+	void SetComputeShader(FMetalComputeShader* InComputeShader);
 	void Dispatch(uint32 ThreadGroupCountX, uint32 ThreadGroupCountY, uint32 ThreadGroupCountZ);
 
 protected:
@@ -180,6 +180,16 @@ protected:
 	 */
 	bool NeedsToSetRenderTarget(const FRHISetRenderTargetsInfo& RenderTargetsInfo);
 
+	/**
+	 * Possibly switch from compute to graphics
+	 */
+	void ConditionalSwitchToGraphics();
+	
+	/**
+	 * Possibly switch from graphics to compute
+	 */
+	void ConditionalSwitchToCompute();
+	
 
 	id<MTLDevice> Device;
 
@@ -232,7 +242,7 @@ protected:
 	
 	class FMetalShaderParameterCache*	ShaderParameters;
 
-	FComputeShaderRHIParamRef CurrentComputeShader;
+	TRefCountPtr<FMetalComputeShader> CurrentComputeShader;
 
 	// the running pipeline state descriptor object
 	FPipelineShadow Pipeline;
diff --git a/Engine/Source/Runtime/IOS/MetalRHI/Private/MetalShaders.cpp b/Engine/Source/Runtime/IOS/MetalRHI/Private/MetalShaders.cpp
--- a/Engine/Source/Runtime/IOS/MetalRHI/Private/MetalShaders.cpp
+++ b/Engine/Source/Runtime/IOS/MetalRHI/Private/MetalShaders.cpp
@@ -181,6 +181,21 @@ TMetalBaseShader<BaseResourceType, ShaderType>::~TMetalBaseShader()
 	[GlslCodeNSString release];
 }
 
+FMetalComputeShader::FMetalComputeShader(const TArray<uint8>& InCode)
+	: TMetalBaseShader<FRHIComputeShader, SF_Compute>(InCode)
+//	, NumThreadsX(0)
+//	, NumThreadsY(0)
+//	, NumThreadsZ(0)
+{
+	NSError* Error;
+	Kernel = [FMetalManager::GetDevice() newComputePipelineStateWithFunction:Function error:&Error];
+	
+	if (Kernel == nil)
+	{
+		NSLog(@"Failed to create compute kernel: %@", Error);
+	}
+}
+
 
 FVertexShaderRHIRef FMetalDynamicRHI::RHICreateVertexShader(const TArray<uint8>& Code)
 {
diff --git a/Engine/Source/Runtime/IOS/MetalRHI/Private/MetalStructuredBuffer.cpp b/Engine/Source/Runtime/IOS/MetalRHI/Private/MetalStructuredBuffer.cpp
--- a/Engine/Source/Runtime/IOS/MetalRHI/Private/MetalStructuredBuffer.cpp
+++ b/Engine/Source/Runtime/IOS/MetalRHI/Private/MetalStructuredBuffer.cpp
@@ -9,17 +9,22 @@ FMetalStructuredBuffer::FMetalStructuredBuffer(uint32 Stride, uint32 Size, FReso
 {
 	check((Size % Stride) == 0);
 
-	// copy any resources to the CPU address
+	Buffer = [FMetalManager::GetDevice() newBufferWithLength:Size options:BUFFER_CACHE_MODE];
+
 	if (ResourceArray)
 	{
-// 		FMemory::Memcpy( , ResourceArray->GetResourceData(), Size);
+		// copy any resources to the CPU address
+		void* LockedMemory = RHILockStructuredBuffer(this, 0, Size, RLM_WriteOnly);
+ 		FMemory::Memcpy(LockedMemory, ResourceArray->GetResourceData(), Size);
 		ResourceArray->Discard();
 	}
+
+	TRACK_OBJECT(Buffer);
 }
 
 FMetalStructuredBuffer::~FMetalStructuredBuffer()
 {
-
+	FMetalManager::ReleaseObject(Buffer);
 }
 
 
@@ -34,10 +39,11 @@ void* FMetalDynamicRHI::RHILockStructuredBuffer(FStructuredBufferRHIParamRef Str
 {
 	DYNAMIC_CAST_METALRESOURCE(StructuredBuffer,StructuredBuffer);
 
-	return NULL;
+	// just return the memory plus the offset
+	return ((uint8*)[StructuredBuffer->Buffer contents]) + Offset;
 }
 
 void FMetalDynamicRHI::RHIUnlockStructuredBuffer(FStructuredBufferRHIParamRef StructuredBufferRHI)
 {
-
+	// nothing to do
 }
diff --git a/Engine/Source/Runtime/IOS/MetalRHI/Private/MetalUAV.cpp b/Engine/Source/Runtime/IOS/MetalRHI/Private/MetalUAV.cpp
--- a/Engine/Source/Runtime/IOS/MetalRHI/Private/MetalUAV.cpp
+++ b/Engine/Source/Runtime/IOS/MetalRHI/Private/MetalUAV.cpp
@@ -9,6 +9,34 @@ FMetalShaderResourceView::~FMetalShaderResourceView()
 	SourceTexture = NULL;
 }
 
+void FMetalUnorderedAccessView::Set(uint32 ResourceIndex)
+{
+	// figure out which one of the resources we need to set
+	FMetalStructuredBuffer* StructuredBuffer = SourceStructuredBuffer.GetReference();
+	FMetalVertexBuffer* VertexBuffer = SourceVertexBuffer.GetReference();
+	FRHITexture* Texture = SourceTexture.GetReference();
+	if (StructuredBuffer)
+	{
+		[FMetalManager::GetComputeContext() setBuffer:StructuredBuffer->Buffer offset:0 atIndex:ResourceIndex];
+	}
+	else if (VertexBuffer)
+	{
+		[FMetalManager::GetComputeContext() setBuffer:VertexBuffer->Buffer offset:VertexBuffer->Offset atIndex:ResourceIndex];
+	}
+	else if (Texture)
+	{
+		FMetalSurface* Surface = GetMetalSurfaceFromRHITexture(Texture);
+		if (Surface != nullptr)
+		{
+			[FMetalManager::GetComputeContext() setTexture:Surface->Texture atIndex:ResourceIndex];
+		}
+		else
+		{
+			[FMetalManager::GetComputeContext() setTexture:nil atIndex:ResourceIndex];
+		}
+	}
+
+}
 
 
 FUnorderedAccessViewRHIRef FMetalDynamicRHI::RHICreateUnorderedAccessView(FStructuredBufferRHIParamRef StructuredBufferRHI, bool bUseUAVCounter, bool bAppendBuffer)
diff --git a/Engine/Source/Runtime/IOS/MetalRHI/Public/MetalResources.h b/Engine/Source/Runtime/IOS/MetalRHI/Public/MetalResources.h
--- a/Engine/Source/Runtime/IOS/MetalRHI/Public/MetalResources.h
+++ b/Engine/Source/Runtime/IOS/MetalRHI/Public/MetalResources.h
@@ -89,14 +89,12 @@ typedef TMetalBaseShader<FRHIGeometryShader, SF_Geometry> FMetalGeometryShader;
 class FMetalComputeShader : public TMetalBaseShader<FRHIComputeShader, SF_Compute>
 {
 public:
-	FMetalComputeShader(const TArray<uint8>& InCode) :
-		TMetalBaseShader<FRHIComputeShader, SF_Compute>(InCode),
-		NumThreadsX(0),
-		NumThreadsY(0),
-		NumThreadsZ(0)
-	{
-	}
-
+	FMetalComputeShader(const TArray<uint8>& InCode);
+	
+	// the state object for a compute shader
+	id <MTLComputePipelineState> Kernel;
+	
+	// thread group counts
 	int32 NumThreadsX;
 	int32 NumThreadsY;
 	int32 NumThreadsZ;
@@ -415,6 +413,8 @@ public:
 	// Destructor
 	~FMetalStructuredBuffer();
 
+	// the actual buffer
+	id<MTLBuffer> Buffer;
 };
 
 
@@ -423,6 +423,9 @@ class FMetalUnorderedAccessView : public FRHIUnorderedAccessView
 {
 public:
 
+	/** Set it into the compute context */
+	void Set(uint32 ResourceIndex);
+	
 	// the potential resources to refer to with the UAV object
 	TRefCountPtr<FMetalStructuredBuffer> SourceStructuredBuffer;
 	TRefCountPtr<FMetalVertexBuffer> SourceVertexBuffer;
diff --git a/Engine/Source/Runtime/RHI/Public/RHIDefinitions.h b/Engine/Source/Runtime/RHI/Public/RHIDefinitions.h
--- a/Engine/Source/Runtime/RHI/Public/RHIDefinitions.h
+++ b/Engine/Source/Runtime/RHI/Public/RHIDefinitions.h
@@ -690,7 +690,7 @@ inline bool RHISupportsSeparateMSAAAndResolveTextures(const EShaderPlatform Plat
 inline bool RHISupportsComputeShaders(const EShaderPlatform Platform)
 {
 	//@todo-rco: Add Metal support
-	return IsFeatureLevelSupported(Platform, ERHIFeatureLevel::SM5);// || (Platform == SP_METAL);
+	return IsFeatureLevelSupported(Platform, ERHIFeatureLevel::SM5) || (Platform == SP_METAL_MRT);
 }
 
 inline bool RHISupportsGeometryShaders(const EShaderPlatform Platform)
