fix for null drawable when entering the background when using metal
UE-6738
#ios

diff --git a/Engine/Source/Runtime/IOS/MetalRHI/Private/MetalManager.cpp b/Engine/Source/Runtime/IOS/MetalRHI/Private/MetalManager.cpp
--- a/Engine/Source/Runtime/IOS/MetalRHI/Private/MetalManager.cpp
+++ b/Engine/Source/Runtime/IOS/MetalRHI/Private/MetalManager.cpp
@@ -509,13 +509,18 @@ void FMetalManager::ConditionalUpdateBackBuffer(FMetalSurface& Surface)
 			// make a drawable object for this frame
 			CurrentDrawable = [[IOSAppDelegate GetDelegate].IOSView MakeDrawable];
 
-			checkf(CurrentDrawable != nil, TEXT("Failed to return the next drawable"));
-			
 			GRenderThreadIdle[ERenderThreadIdleTypes::WaitingForGPUPresent] += FPlatformTime::Cycles() - IdleStart;
 			GRenderThreadNumIdle[ERenderThreadIdleTypes::WaitingForGPUPresent]++;
 
 			// set the texture into the backbuffer
-			Surface.Texture = CurrentDrawable.texture;
+            if (CurrentDrawable != nil)
+            {
+                Surface.Texture = CurrentDrawable.texture;
+            }
+            else
+            {
+                Surface.Texture = nil;
+            }
 		}
 	}
 }
@@ -600,6 +605,7 @@ void FMetalManager::SetRenderTargetsInfo(const FRHISetRenderTargetsInfo& RenderT
 	RenderPass.visibilityResultBuffer = QueryBuffer.Buffer;
 
 	// default to non-msaa
+    int32 OldCount = Pipeline.SampleCount;
 	Pipeline.SampleCount = 0;
 
 	for (uint32 RenderTargetIndex = 0; RenderTargetIndex < MaxMetalRenderTargets; RenderTargetIndex++)
@@ -615,6 +621,12 @@ void FMetalManager::SetRenderTargetsInfo(const FRHISetRenderTargetsInfo& RenderT
 		
 			// if this is the back buffer, make sure we have a usable drawable
 			ConditionalUpdateBackBuffer(Surface);
+            
+            if (Surface.Texture == nil)
+            {
+                Pipeline.SampleCount = OldCount;
+                return;
+            }
 
 			// user code generally passes -1 as a default, but we need 0
 			uint32 ArraySliceIndex = RenderTargetView.ArraySliceIndex == 0xFFFFFFFF ? 0 : RenderTargetView.ArraySliceIndex;
