UE4 - Metal - Fix bad vertex declaration size

diff --git a/Engine/Source/Runtime/IOS/MetalRHI/Private/MetalVertexDeclaration.cpp b/Engine/Source/Runtime/IOS/MetalRHI/Private/MetalVertexDeclaration.cpp
--- a/Engine/Source/Runtime/IOS/MetalRHI/Private/MetalVertexDeclaration.cpp
+++ b/Engine/Source/Runtime/IOS/MetalRHI/Private/MetalVertexDeclaration.cpp
@@ -42,7 +42,7 @@ static uint32 TranslateElementTypeToSize(EVertexElementType Type)
 		case VET_Color:			return 4;
 		case VET_Short2:		return 4;
 		case VET_Short4:		return 8;
-		case VET_Short2N:		return 8;
+		case VET_Short2N:		return 4;
 		case VET_Half2:			return 4;
 		case VET_Half4:			return 8;
 		default:				UE_LOG(LogMetal, Fatal, TEXT("Unknown vertex element type!")); return 0;
