Core - UE/UE-6813 - Confusing API: We have both FMemory::MemZero() and FMemory:Memzero()

diff --git a/Engine/Source/Developer/MeshSimplifier/Private/Cache.h b/Engine/Source/Developer/MeshSimplifier/Private/Cache.h
--- a/Engine/Source/Developer/MeshSimplifier/Private/Cache.h
+++ b/Engine/Source/Developer/MeshSimplifier/Private/Cache.h
@@ -34,7 +34,7 @@ FORCEINLINE TCacheDirect<T, Size>::TCacheDirect()
 {
 	static_assert((Size & (Size - 1)) == 0, "Size must be power of 2.");
 
-	FMemory::MemSet( Keys, 0xff );
+	FMemory::Memset( Keys, 0xff );
 }
 
 template< typename T, uint32 Size >
@@ -94,10 +94,10 @@ FORCEINLINE TCache4Way<T, size>::TCache4Way()
 {
 	static_assert((Size & (Size - 1)) == 0, "Size must be power of 2.");
 
-	FMemory::MemSet( Keys, 0xff );
+	FMemory::Memset( Keys, 0xff );
 
 	uint8 StartingOrder = (3 << 6) | (2 << 4) | (1 << 2) | 0;
-	FMemory::MemSet( order, StartingOrder );
+	FMemory::Memset( order, StartingOrder );
 }
 
 template< typename T, uint32 size >
diff --git a/Engine/Source/Developer/MeshUtilities/Private/LayoutUV.cpp b/Engine/Source/Developer/MeshUtilities/Private/LayoutUV.cpp
--- a/Engine/Source/Developer/MeshUtilities/Private/LayoutUV.cpp
+++ b/Engine/Source/Developer/MeshUtilities/Private/LayoutUV.cpp
@@ -160,7 +160,7 @@ void FLayoutUV::FindCharts( const TMultiMap<int32,int32>& OverlappingCorners )
 		Chart.MaxUV = FVector2D( -FLT_MAX, -FLT_MAX );
 		Chart.UVArea = 0.0f;
 		Chart.WorldScale = FVector2D::ZeroVector;
-		FMemory::MemSet( Chart.Join, 0xff );
+		FMemory::Memset( Chart.Join, 0xff );
 
 		Chart.FirstTri = Tri;
 
diff --git a/Engine/Source/Developer/Profiler/Private/Widgets/SFiltersAndPresets.cpp b/Engine/Source/Developer/Profiler/Private/Widgets/SFiltersAndPresets.cpp
--- a/Engine/Source/Developer/Profiler/Private/Widgets/SFiltersAndPresets.cpp
+++ b/Engine/Source/Developer/Profiler/Private/Widgets/SFiltersAndPresets.cpp
@@ -580,7 +580,7 @@ SFiltersAndPresets::SFiltersAndPresets()
 	, SortingMode( EStatGroupingOrSortingMode::StatName )
 	, bExpansionSaved( false )
 {
-	FMemory::MemSet( bStatTypeIsVisible, 1 );
+	FMemory::Memset( bStatTypeIsVisible, 1 );
 }
 
 SFiltersAndPresets::~SFiltersAndPresets()
diff --git a/Engine/Source/Editor/Persona/Private/SAnimationBlendSpace.h b/Engine/Source/Editor/Persona/Private/SAnimationBlendSpace.h
--- a/Engine/Source/Editor/Persona/Private/SAnimationBlendSpace.h
+++ b/Engine/Source/Editor/Persona/Private/SAnimationBlendSpace.h
@@ -90,8 +90,8 @@ struct FTriangle
 
 	FTriangle(FTriangle & Copy)
 	{
-		FMemory::MemCopy(Vertices, Copy.Vertices);
-		FMemory::MemCopy(Edges, Copy.Edges);
+		FMemory::Memcpy(Vertices, Copy.Vertices);
+		FMemory::Memcpy(Edges, Copy.Edges);
 		Center = Copy.Center;
 
 		Vertices[0]->AddTriangle(this);
diff --git a/Engine/Source/Editor/UnrealEd/Private/Factories/EditorFactories.cpp b/Engine/Source/Editor/UnrealEd/Private/Factories/EditorFactories.cpp
--- a/Engine/Source/Editor/UnrealEd/Private/Factories/EditorFactories.cpp
+++ b/Engine/Source/Editor/UnrealEd/Private/Factories/EditorFactories.cpp
@@ -148,7 +148,7 @@ UObject* UTexture2DFactoryNew::FactoryCreateNew( UClass* InClass, UObject* InPar
 		Object->Source.GetMipData( TexturePixels, 0 );
 
 		uint8* DestData = Object->Source.LockMip(0);
-		FMemory::MemSet(DestData, TexturePixels.Num() * sizeof( uint8 ) );
+		FMemory::Memset(DestData, TexturePixels.Num() * sizeof( uint8 ), 255 );
 		Object->Source.UnlockMip(0);
 
 		Object->PostEditChange();
diff --git a/Engine/Source/Runtime/Engine/Private/Animation/AnimSequenceBase.cpp b/Engine/Source/Runtime/Engine/Private/Animation/AnimSequenceBase.cpp
--- a/Engine/Source/Runtime/Engine/Private/Animation/AnimSequenceBase.cpp
+++ b/Engine/Source/Runtime/Engine/Private/Animation/AnimSequenceBase.cpp
@@ -145,7 +145,7 @@ void FFloatCurve::UpdateOrAddKey(float NewKey, float CurrentTime)
 // we don't want to have = operator. This only copies curves, but leaving naming and everything else intact. 
 void FVectorCurve::CopyCurve(FVectorCurve& SourceCurve)
 {
-	FMemory::MemCopy(FloatCurves, SourceCurve.FloatCurves);
+	FMemory::Memcpy(FloatCurves, SourceCurve.FloatCurves);
 }
 
 FVector FVectorCurve::Evaluate(float CurrentTime, float BlendWeight) const
diff --git a/Engine/Source/Runtime/IOS/MetalRHI/Private/MetalManager.cpp b/Engine/Source/Runtime/IOS/MetalRHI/Private/MetalManager.cpp
--- a/Engine/Source/Runtime/IOS/MetalRHI/Private/MetalManager.cpp
+++ b/Engine/Source/Runtime/IOS/MetalRHI/Private/MetalManager.cpp
@@ -188,12 +188,12 @@ FMetalManager::FMetalManager()
 		CurrentColorRenderTextures[Index] = nil;
 		PreviousColorRenderTextures[Index] = nil;
 	}
-	FMemory::MemSet(CurrentRenderTargetsViewInfo, 0);
-	FMemory::MemSet(PreviousRenderTargetsViewInfo, 0);
-	FMemory::MemSet(CurrentDepthViewInfo, 0);
-	FMemory::MemSet(PreviousDepthViewInfo, 0);
-	FMemory::MemSet(CurrentStencilViewInfo, 0);
-	FMemory::MemSet(PreviousStencilViewInfo, 0);
+	FMemory::Memzero(CurrentRenderTargetsViewInfo);
+	FMemory::Memzero(PreviousRenderTargetsViewInfo);
+	FMemory::Memzero(CurrentDepthViewInfo);
+	FMemory::Memzero(PreviousDepthViewInfo);
+	FMemory::Memzero(CurrentStencilViewInfo);
+	FMemory::Memzero(PreviousStencilViewInfo);
 
 	CommandQueue = [Device newCommandQueue];
 
