UE4 - Metal compile fix

diff --git a/Engine/Source/Runtime/IOS/MetalRHI/Private/MetalQuery.cpp b/Engine/Source/Runtime/IOS/MetalRHI/Private/MetalQuery.cpp
--- a/Engine/Source/Runtime/IOS/MetalRHI/Private/MetalQuery.cpp
+++ b/Engine/Source/Runtime/IOS/MetalRHI/Private/MetalQuery.cpp
@@ -76,3 +76,11 @@ bool FMetalDynamicRHI::RHIGetRenderQueryResult(FRenderQueryRHIParamRef QueryRHI,
 	OutNumPixels = *QueryMemory;
 	return true;
 }
+
+void FMetalDynamicRHI::RHIBeginOcclusionQueryBatch()
+{
+}
+
+void FMetalDynamicRHI::RHIEndOcclusionQueryBatch()
+{
+}
