Fix Metal CIS

diff --git a/Engine/Source/Runtime/IOS/MetalRHI/Private/MetalCommands.cpp b/Engine/Source/Runtime/IOS/MetalRHI/Private/MetalCommands.cpp
--- a/Engine/Source/Runtime/IOS/MetalRHI/Private/MetalCommands.cpp
+++ b/Engine/Source/Runtime/IOS/MetalRHI/Private/MetalCommands.cpp
@@ -463,11 +463,19 @@ void FMetalDynamicRHI::RHISetBlendState(FBlendStateRHIParamRef NewStateRHI, cons
 
 
 void FMetalDynamicRHI::RHISetRenderTargets(uint32 NumSimultaneousRenderTargets, const FRHIRenderTargetView* NewRenderTargets,
-	FTextureRHIParamRef NewDepthStencilTargetRHI, uint32 NumUAVs, const FUnorderedAccessViewRHIParamRef* UAVs)
+	const FRHIDepthRenderTargetView* NewDepthStencilTargetRHI, uint32 NumUAVs, const FUnorderedAccessViewRHIParamRef* UAVs)
 {
 	FMetalManager* Manager = FMetalManager::Get();
+	FRHIDepthRenderTargetView DepthView;
+	if (NewDepthStencilTargetRHI)
+	{
+		DepthView = *NewDepthStencilTargetRHI;
+	}
+	else
+	{
+		DepthView = FRHIDepthRenderTargetView(FTextureRHIParamRef(), ERenderTargetLoadAction::EClear, ERenderTargetStoreAction::ENoAction);
+	}
 
-	FRHIDepthRenderTargetView DepthView(NewDepthStencilTargetRHI, ERenderTargetLoadAction::EClear, ERenderTargetStoreAction::ENoAction);
 	FRHISetRenderTargetsInfo Info(NumSimultaneousRenderTargets, NewRenderTargets, DepthView);
 	RHISetRenderTargetsAndClear(Info);
 	
