Fix another metal CIS

diff --git a/Engine/Source/Runtime/IOS/MetalRHI/Private/MetalCommands.cpp b/Engine/Source/Runtime/IOS/MetalRHI/Private/MetalCommands.cpp
--- a/Engine/Source/Runtime/IOS/MetalRHI/Private/MetalCommands.cpp
+++ b/Engine/Source/Runtime/IOS/MetalRHI/Private/MetalCommands.cpp
@@ -840,9 +840,9 @@ void FMetalDynamicRHI::RHIClearMRT(bool bClearColor,int32 NumClearColors,const F
 
 		RHIPopEvent();
 	}
-}
+}											
 
-void FMetalDynamicRHI::RHIBindClearMRTValues(bool bClearColor, int32 NumClearColors, const FLinearColor* ClearColorArray, bool bClearDepth, float Depth, bool bClearStencil, uint32 Stencil, bool bCheckErrors)
+void FMetalDynamicRHI::RHIBindClearMRTValues(bool bClearColor, int32 NumClearColors, const FLinearColor* ClearColorArray, bool bClearDepth, float Depth, bool bClearStencil, uint32 Stencil)
 {
 	// Not necessary
 }
