fix for startup movies not showing when running Metal
UE-7636
#ios

diff --git a/Engine/Source/Runtime/IOS/MetalRHI/Private/MetalManager.cpp b/Engine/Source/Runtime/IOS/MetalRHI/Private/MetalManager.cpp
--- a/Engine/Source/Runtime/IOS/MetalRHI/Private/MetalManager.cpp
+++ b/Engine/Source/Runtime/IOS/MetalRHI/Private/MetalManager.cpp
@@ -397,7 +397,7 @@ void FMetalManager::EndFrame(bool bPresent)
 	// enqueue a present if desired
 	if (CurrentDrawable)
 	{
-		if (bPresent && GFrameCounter > 30)
+		if (bPresent)
 		{
 			[CurrentCommandBuffer waitUntilScheduled];
 			
