- Cleaned up some tooltips, and disabled MetalMRT for rocket builds (not up to snuff) [UE-9253]

diff --git a/Engine/Source/Developer/iOS/IOSPlatformEditor/Private/IOSTargetSettingsCustomization.cpp b/Engine/Source/Developer/iOS/IOSPlatformEditor/Private/IOSTargetSettingsCustomization.cpp
--- a/Engine/Source/Developer/iOS/IOSPlatformEditor/Private/IOSTargetSettingsCustomization.cpp
+++ b/Engine/Source/Developer/iOS/IOSPlatformEditor/Private/IOSTargetSettingsCustomization.cpp
@@ -549,11 +549,6 @@ void FIOSTargetSettingsCustomization::BuildPListSection(IDetailLayoutBuilder& De
 		 ]
 	 ];
 
-	// Show properties that are gated by the plist being present and writable
-	FSimpleDelegate BundleModifiedDelegate = FSimpleDelegate::CreateRaw(this, &FIOSTargetSettingsCustomization::FindRequiredFiles);
-	FGameProjectGenerationModule& GameProjectModule = FModuleManager::LoadModuleChecked<FGameProjectGenerationModule>(TEXT("GameProjectGeneration"));
-	bool bHasCode = GameProjectModule.Get().ProjectHasCodeFiles();
-
 #define SETUP_NONROCKET_PROP(PropName, Category, Tip, DisabledTip) \
 	{ \
 		TSharedRef<IPropertyHandle> PropertyHandle = DetailLayout.GetProperty(GET_MEMBER_NAME_CHECKED(UIOSRuntimeSettings, PropName)); \
@@ -562,40 +557,6 @@ void FIOSTargetSettingsCustomization::BuildPListSection(IDetailLayoutBuilder& De
 			.ToolTip(!FRocketSupport::IsRocket() ? Tip : DisabledTip); \
 	}
 
-#define SETUP_PLIST_PROP(PropName, Category, Tip) \
-	{ \
-		TSharedRef<IPropertyHandle> PropertyHandle = DetailLayout.GetProperty(GET_MEMBER_NAME_CHECKED(UIOSRuntimeSettings, PropName)); \
-		Category.AddProperty(PropertyHandle) \
-			.ToolTip(Tip); \
-	}
-
-#define SETUP_STATUS_PROP(PropName, Category, Tip) \
-	{ \
-	TSharedRef<IPropertyHandle> PropertyHandle = DetailLayout.GetProperty(GET_MEMBER_NAME_CHECKED(UIOSRuntimeSettings, PropName)); \
-	PropertyHandle->SetOnPropertyValueChanged(BundleModifiedDelegate); \
-	Category.AddProperty(PropertyHandle) \
-		.ToolTip(Tip); \
-	}
-
-	SETUP_PLIST_PROP(BundleDisplayName, BundleCategory, LOCTEXT("BundleDisplayNameToolTip", "Specifies the the display name for the application. This will be displayed under the icon on the device."));
-	SETUP_PLIST_PROP(BundleName, BundleCategory, LOCTEXT("BundleNameToolTip", "Specifies the the name of the application bundle. This is the short name for the application bundle."));
-	SETUP_STATUS_PROP(BundleIdentifier, BundleCategory, LOCTEXT("BundleIdentifierToolTip", "Specifies the bundle identifier for the application."));
-	SETUP_PLIST_PROP(VersionInfo, BundleCategory, LOCTEXT("VersionInfoToolTip", "Specifies the version for the application."));
-	SETUP_PLIST_PROP(bSupportsPortraitOrientation, OrientationCategory, LOCTEXT("SupportsPortraitOrientationToolTip", "Supports default portrait orientation. Landscape will not be supported."));
-	SETUP_PLIST_PROP(bSupportsUpsideDownOrientation, OrientationCategory, LOCTEXT("SupportsUpsideDownOrientationToolTip", "Supports upside down portrait orientation. Landscape will not be supported."));
-	SETUP_PLIST_PROP(bSupportsLandscapeLeftOrientation, OrientationCategory, LOCTEXT("SupportsLandscapeLeftOrientationToolTip", "Supports left landscape orientation. Protrait will not be supported."));
-	SETUP_PLIST_PROP(bSupportsLandscapeRightOrientation, OrientationCategory, LOCTEXT("SupportsLandscapeRightOrientationToolTip", "Supports right landscape orientation. Protrait will not be supported."));
-	
-	SETUP_PLIST_PROP(bSupportsMetal, RenderCategory, LOCTEXT("SupportsMetalToolTip", "Whether or not to add support for Metal API (requires IOS8 and A7 processors)."));
-	SETUP_PLIST_PROP(bSupportsOpenGLES2, RenderCategory, LOCTEXT("SupportsOpenGLES2ToolTip", "Whether or not to add support for OpenGL ES2 (if this is false, then your game should specify minimum IOS8 version and use \"metal\" instead of \"opengles-2\" in UIRequiredDeviceCapabilities)"));
-
-	SETUP_PLIST_PROP(bSupportsIPad, DeviceCategory, LOCTEXT("SupportsIPadToolTip", "Whether or not to add support for iPad devices"));
-	SETUP_PLIST_PROP(bSupportsIPhone, DeviceCategory, LOCTEXT("SupportsIPhoneToolTip", "Whether or not to add support for iPhone devices"));
-
-	SETUP_PLIST_PROP(MinimumiOSVersion, OSInfoCategory, LOCTEXT("MinimumiOSVersionToolTip", "Minimum iOS version this game supports"));
-
-	SETUP_PLIST_PROP(AdditionalPlistData, ExtraCategory, LOCTEXT("AdditionalPlistDataToolTip", "Any additional plist key/value data utilizing \\n for a new line"));
-
 	SETUP_NONROCKET_PROP(bDevForArmV7, BuildCategory, LOCTEXT("DevForArmV7ToolTip", "Enable ArmV7 support? (this will be used if all type are unchecked)"), FIOSTargetSettingsCustomizationConstants::DisabledTip);
 	SETUP_NONROCKET_PROP(bDevForArm64, BuildCategory, LOCTEXT("DevForArm64ToolTip", "Enable Arm64 support?"), FIOSTargetSettingsCustomizationConstants::DisabledTip);
 	SETUP_NONROCKET_PROP(bDevForArmV7S, BuildCategory, LOCTEXT("DevForArmV7SToolTip", "Enable ArmV7s support?"), FIOSTargetSettingsCustomizationConstants::DisabledTip);
@@ -603,7 +564,8 @@ void FIOSTargetSettingsCustomization::BuildPListSection(IDetailLayoutBuilder& De
 	SETUP_NONROCKET_PROP(bShipForArm64, BuildCategory, LOCTEXT("ShipForArm64ToolTip", "Enable Arm64 support?"), FIOSTargetSettingsCustomizationConstants::DisabledTip);
 	SETUP_NONROCKET_PROP(bShipForArmV7S, BuildCategory, LOCTEXT("ShipForArmV7SToolTip", "Enable ArmV7s support?"), FIOSTargetSettingsCustomizationConstants::DisabledTip);
 
-#undef SETUP_PLIST_PROP
+	SETUP_NONROCKET_PROP(bSupportsMetalMRT, RenderCategory, LOCTEXT("MetalMRTToolTip", "Whether or not to add support for deferred rendering Metal API (requires IOS8 and A8 processors). This is very much a work in progress, be aware!"), FIOSTargetSettingsCustomizationConstants::DisabledTip);
+
 #undef SETUP_NONROCKET_PROP
 }
 
diff --git a/Engine/Source/Runtime/IOS/IOSRuntimeSettings/Classes/IOSRuntimeSettings.h b/Engine/Source/Runtime/IOS/IOSRuntimeSettings/Classes/IOSRuntimeSettings.h
--- a/Engine/Source/Runtime/IOS/IOSRuntimeSettings/Classes/IOSRuntimeSettings.h
+++ b/Engine/Source/Runtime/IOS/IOSRuntimeSettings/Classes/IOSRuntimeSettings.h
@@ -96,39 +96,39 @@ public:
 	UPROPERTY(GlobalConfig, EditAnywhere, Category = Online)
 	uint32 bEnableGameCenterSupport : 1;
 	
-	// Whether or not to add support for Metal API (requires IOS8 and A7 processors)
-	UPROPERTY(GlobalConfig, EditAnywhere, Category = Rendering, meta = (DisplayName = "Supports Forward Rendering with Metal (A7 and up devices)"))
+	// Whether or not to add support for Metal API (requires IOS8 and A7 processors).
+	UPROPERTY(GlobalConfig, EditAnywhere, Category = Rendering, meta = (DisplayName = "Support Forward Rendering with Metal (A7 and up devices)"))
 	bool bSupportsMetal;
 
-	// Whether or not to add support for deferred rendering Metal API (requires IOS8 and A7 processors)
-	UPROPERTY(GlobalConfig, EditAnywhere, Category = Rendering, meta = (DisplayName = "Supports Deferred Rendering with Metal (A8 and up devices)"))
+	// Whether or not to add support for deferred rendering Metal API (requires IOS8 and A8 processors)
+	UPROPERTY(GlobalConfig, EditAnywhere, Category = Rendering, meta = (DisplayName = "[Work in Progress] Support Deferred Rendering with Metal (A8 and up devices)"))
 	bool bSupportsMetalMRT;
 	
-	// Whether or not to add support for OpenGL ES2 (if this is false, then your game should specify minimum IOS8 version and use "metal" instead of "opengles-2" in UIRequiredDeviceCapabilities)
+	// Whether or not to add support for OpenGL ES2 (if this is false, then your game should specify minimum IOS8 version)
 	UPROPERTY(GlobalConfig, EditAnywhere, Category = Rendering)
 	bool bSupportsOpenGLES2;
 
-	// Enable ArmV7 support? (this will be used if all type are unchecked) [CURRENTLY AVAILABLE ON MAC ONLY OR FULL SOURCE BUILD]
+	// Enable ArmV7 support? (this will be used if all type are unchecked)
 	UPROPERTY(GlobalConfig, EditAnywhere, Category = Build, meta = (DisplayName = "Support armv7 in Development"))
 	bool bDevForArmV7;
 
-	// Enable Arm64 support? [CURRENTLY AVAILABLE ON MAC ONLY OR FULL SOURCE BUILD]
+	// Enable Arm64 support?
 	UPROPERTY(GlobalConfig, EditAnywhere, Category = Build, meta = (DisplayName = "Support arm64 in Development"))
 	bool bDevForArm64;
 
-	// Enable ArmV7s support? [CURRENTLY AVAILABLE ON MAC ONLY OR FULL SOURCE BUILD]
+	// Enable ArmV7s support?
 	UPROPERTY(GlobalConfig, EditAnywhere, Category = Build, meta = (DisplayName = "Support armv7s in Development"))
 	bool bDevForArmV7S;
 
-	// Enable ArmV7 support for shipping build? (this will be used if all type are unchecked) [CURRENTLY AVAILABLE ON MAC ONLY OR FULL SOURCE BUILD]
+	// Enable ArmV7 support? (this will be used if all type are unchecked)
 	UPROPERTY(GlobalConfig, EditAnywhere, Category = Build, meta = (DisplayName = "Support armv7 in Shipping"))
 	bool bShipForArmV7;
 
-	// Enable Arm64 support for shipping build? [CURRENTLY AVAILABLE ON MAC ONLY OR FULL SOURCE BUILD]
+	// Enable Arm64 support?
 	UPROPERTY(GlobalConfig, EditAnywhere, Category = Build, meta = (DisplayName = "Support arm64 in Shipping"))
 	bool bShipForArm64;
 
-	// Enable ArmV7s support for shipping build? [CURRENTLY AVAILABLE ON MAC ONLY OR FULL SOURCE BUILD]
+	// Enable ArmV7s support?
 	UPROPERTY(GlobalConfig, EditAnywhere, Category = Build, meta = (DisplayName = "Support armv7s in Shipping"))
 	bool bShipForArmV7S;
 	
@@ -152,35 +152,35 @@ public:
 	UPROPERTY(GlobalConfig, EditAnywhere, Category = "Build", meta = (EditCondition = "bUseRSync", DisplayName = "Override existing SSH permissions file"))
 	FIOSBuildResourceFilePath SSHPrivateKeyOverridePath;
 
-	// Does the application support portrait orientation?
+	// Supports default portrait orientation. Landscape will not be supported.
 	UPROPERTY(GlobalConfig, EditAnywhere, Category = DeviceOrientations)
 	uint32 bSupportsPortraitOrientation : 1;
 
-	// Does the application support upside down orientation?
+	// Supports upside down portrait orientation. Landscape will not be supported.
 	UPROPERTY(GlobalConfig, EditAnywhere, Category = DeviceOrientations)
 	uint32 bSupportsUpsideDownOrientation : 1;
 
-	// Does the application support landscape left orientation?
+	// Supports left landscape orientation. Protrait will not be supported.
 	UPROPERTY(GlobalConfig, EditAnywhere, Category = DeviceOrientations)
 	uint32 bSupportsLandscapeLeftOrientation : 1;
 
-	// Does the application support landscape right orientation?
+	// Supports right landscape orientation. Protrait will not be supported.
 	UPROPERTY(GlobalConfig, EditAnywhere, Category = DeviceOrientations)
 	uint32 bSupportsLandscapeRightOrientation : 1;
 
-	// Bundle Display Name
+	// Specifies the the display name for the application. This will be displayed under the icon on the device.
 	UPROPERTY(GlobalConfig, EditAnywhere, Category = BundleInfo)
 	FString BundleDisplayName;
 
-	// Bundle Name
+	// Specifies the the name of the application bundle. This is the short name for the application bundle.
 	UPROPERTY(GlobalConfig, EditAnywhere, Category = BundleInfo)
 	FString BundleName;
 
-	// Bundle identifier
+	// Specifies the bundle identifier for the application.
 	UPROPERTY(GlobalConfig, EditAnywhere, Category = BundleInfo)
 	FString BundleIdentifier;
 
-	// version info
+	// Specifies the version for the application.
 	UPROPERTY(GlobalConfig, EditAnywhere, Category = BundleInfo)
 	FString VersionInfo;
     
@@ -188,19 +188,19 @@ public:
     UPROPERTY(GlobalConfig, EditAnywhere, Category = PowerUsage)
     TEnumAsByte<EPowerUsageFrameRateLock> FrameRateLock;
 
-	/** Set the minimum iOS setting */
+	// Minimum iOS version this game supports
 	UPROPERTY(GlobalConfig, EditAnywhere, Category = OSInfo)
 	TEnumAsByte<EIOSVersion> MinimumiOSVersion;
 
-	/** Does the application support iPad */
+	// Whether or not to add support for iPad devices
 	UPROPERTY(GlobalConfig, EditAnywhere, Category = DeviceUsage)
 	uint32 bSupportsIPad : 1;
 
-	// Does the application support iPhone */
+	// Whether or not to add support for iPhone devices
 	UPROPERTY(GlobalConfig, EditAnywhere, Category = DeviceUsage)
 	uint32 bSupportsIPhone : 1;
 
-	// extra data to be stored in the plist
+	// Any additional plist key/value data utilizing \n for a new line
 	UPROPERTY(GlobalConfig, EditAnywhere, Category = ExtraData)
 	FString AdditionalPlistData;
 
