Core - Replaced FPlatformProcess::CreateSyncEvent with FPlatformProcess::GetSynchEventFromPool

diff --git a/Engine/Source/Runtime/IOS/MetalRHI/Private/MetalManager.cpp b/Engine/Source/Runtime/IOS/MetalRHI/Private/MetalManager.cpp
--- a/Engine/Source/Runtime/IOS/MetalRHI/Private/MetalManager.cpp
+++ b/Engine/Source/Runtime/IOS/MetalRHI/Private/MetalManager.cpp
@@ -217,7 +217,7 @@ FMetalManager::FMetalManager()
 	FrameReadyEvent = NULL;
 	if( FIOSPlatformRHIFramePacer::IsEnabled() )
 	{
-		FrameReadyEvent = FPlatformProcess::CreateSynchEvent();
+		FrameReadyEvent = FPlatformProcess::GetSynchEventFromPool();
 		FIOSPlatformRHIFramePacer::InitWithEvent( FrameReadyEvent );
 	}
 	
@@ -313,7 +313,7 @@ void FMetalManager::CreateCurrentCommandBuffer(bool bWait)
 	// make a new event if we don't have enough
 	if (CurrentQueryEventIndex >= QueryEvents[WhichFreeList].Num())
 	{
-		LocalEvent = FPlatformProcess::CreateSynchEvent(true);
+		LocalEvent = FPlatformProcess::GetSynchEventFromPool(true);
 		CurrentQueryEventIndex = QueryEvents[WhichFreeList].Add(LocalEvent);
 	}
 	// otherwise, reuse an old one
diff --git a/Engine/Source/Runtime/OpenGLDrv/Private/IOS/IOSOpenGL.cpp b/Engine/Source/Runtime/OpenGLDrv/Private/IOS/IOSOpenGL.cpp
--- a/Engine/Source/Runtime/OpenGLDrv/Private/IOS/IOSOpenGL.cpp
+++ b/Engine/Source/Runtime/OpenGLDrv/Private/IOS/IOSOpenGL.cpp
@@ -85,7 +85,7 @@ struct FPlatformOpenGLDevice
         {
             if( FrameReadyEvent == nullptr )
             {
-                FrameReadyEvent = FPlatformProcess::CreateSynchEvent();
+                FrameReadyEvent = FPlatformProcess::GetSynchEventFromPool();
                 FIOSPlatformRHIFramePacer::InitWithEvent( FrameReadyEvent );
             }
         }
