UE4 - Fix missing RHI vendor Id

diff --git a/Engine/Source/Runtime/IOS/MetalRHI/Private/MetalRHI.cpp b/Engine/Source/Runtime/IOS/MetalRHI/Private/MetalRHI.cpp
--- a/Engine/Source/Runtime/IOS/MetalRHI/Private/MetalRHI.cpp
+++ b/Engine/Source/Runtime/IOS/MetalRHI/Private/MetalRHI.cpp
@@ -99,7 +99,9 @@ FMetalDynamicRHI::FMetalDynamicRHI()
 	// we cannot render to a volume texture without geometry shader support
 	GSupportsVolumeTextureRendering = false;
 
-// 	GDrawUPVertexCheckCount = MAX_uint16;
+	//@todo-rco: Query name from API
+	GRHIAdapterName = TEXT("Metal");
+	GRHIVendorId = 1; // non-zero to avoid asserts
 
 	// Initialize the platform pixel format map.
 	GPixelFormats[PF_Unknown			].PlatformFormat	= MTLPixelFormatInvalid;
