fixed iOS metal compile

diff --git a/Engine/Source/Runtime/IOS/MetalRHI/Private/MetalCommands.cpp b/Engine/Source/Runtime/IOS/MetalRHI/Private/MetalCommands.cpp
--- a/Engine/Source/Runtime/IOS/MetalRHI/Private/MetalCommands.cpp
+++ b/Engine/Source/Runtime/IOS/MetalRHI/Private/MetalCommands.cpp
@@ -516,7 +516,7 @@ void FMetalDynamicRHI::RHISetRenderTargetsAndClear(const FRHISetRenderTargetsInf
 			GetMetalRTStoreAction(RenderTargetsInfo.DepthStencilRenderTarget.DepthStoreAction), 
 			RenderTargetsInfo.DepthClearValue,
 			GetMetalRTLoadAction(RenderTargetsInfo.DepthStencilRenderTarget.StencilLoadAction),
-			GetMetalRTStoreAction(RenderTargetsInfo.DepthStencilRenderTarget.StencilStoreAction),
+			GetMetalRTStoreAction(RenderTargetsInfo.DepthStencilRenderTarget.GetStencilStoreAction()),
 			RenderTargetsInfo.StencilClearValue);
 	}
 	else
diff --git a/Engine/Source/Runtime/IOS/MetalRHI/Private/MetalManager.cpp b/Engine/Source/Runtime/IOS/MetalRHI/Private/MetalManager.cpp
--- a/Engine/Source/Runtime/IOS/MetalRHI/Private/MetalManager.cpp
+++ b/Engine/Source/Runtime/IOS/MetalRHI/Private/MetalManager.cpp
@@ -713,7 +713,7 @@ void FMetalManager::SetRenderTargetsInfo(const FRHISetRenderTargetsInfo& RenderT
 			// set up the stencil attachment
 			StencilAttachment.texture = Surface.StencilTexture;
 			StencilAttachment.loadAction = GetMetalRTLoadAction(RenderTargetsInfo.DepthStencilRenderTarget.StencilLoadAction);
-			StencilAttachment.storeAction = GetMetalRTStoreAction(RenderTargetsInfo.DepthStencilRenderTarget.StencilStoreAction);
+			StencilAttachment.storeAction = GetMetalRTStoreAction(RenderTargetsInfo.DepthStencilRenderTarget.GetStencilStoreAction());
 			StencilAttachment.clearStencil = RenderTargetsInfo.StencilClearValue;
 
 			Pipeline.StencilTargetFormat = StencilAttachment.texture.pixelFormat;
diff --git a/Engine/Source/Runtime/IOS/MetalRHI/Private/MetalManager.h b/Engine/Source/Runtime/IOS/MetalRHI/Private/MetalManager.h
--- a/Engine/Source/Runtime/IOS/MetalRHI/Private/MetalManager.h
+++ b/Engine/Source/Runtime/IOS/MetalRHI/Private/MetalManager.h
@@ -103,7 +103,7 @@ public:
 	void SetBoundShaderState(FMetalBoundShaderState* BoundShaderState);
 	void SetCurrentRenderTarget(FMetalSurface* RenderSurface, int32 RenderTargetIndex, uint32 MipIndex, uint32 ArraySliceIndex, MTLLoadAction LoadAction, MTLStoreAction StoreAction, int32 TotalNumRenderTargets);
 	void SetCurrentDepthStencilTarget(FMetalSurface* RenderSurface, MTLLoadAction DepthLoadAction=MTLLoadActionDontCare, MTLStoreAction DepthStoreAction=MTLStoreActionDontCare, float ClearDepthValue=0, 
-		MTLLoadAction StencilLoadAction=MTLLoadActionDontCare, MTLStoreAction StencilStoreAction=MTLStoreActionDontCare, uint8 ClearStencilValue=0);
+		MTLLoadAction StencilLoadAction=MTLLoadActionDontCare, MTLStoreAction InStencilStoreAction=MTLStoreActionDontCare, uint8 ClearStencilValue=0);
 	
 	/**
 	 * Set the color, depth and stencil render targets, and then make the new command buffer/encoder
