UERNDR-54 - Support switching device orientation on the fly.
Added iOS rotation support and ES2 + metal handling for backbuffer  resizing.
Fixed ogl issue that could delete gltexture when render buffers are destructed.
Fixed issue with ES2 render target dirty state check that could result in some render targets not being set.

diff --git a/Engine/Source/Runtime/Core/Private/IOS/IOSApplication.cpp b/Engine/Source/Runtime/Core/Private/IOS/IOSApplication.cpp
--- a/Engine/Source/Runtime/Core/Private/IOS/IOSApplication.cpp
+++ b/Engine/Source/Runtime/Core/Private/IOS/IOSApplication.cpp
@@ -9,6 +9,10 @@
 #include "IInputInterface.h"
 #include "IInputDevice.h"
 
+
+FCriticalSection FIOSApplication::CriticalSection;
+bool FIOSApplication::bOrientationChanged = false;
+
 FIOSApplication* FIOSApplication::CreateIOSApplication()
 {
 	return new FIOSApplication();
@@ -75,6 +79,20 @@ void FIOSApplication::PollGameDeviceState( const float TimeDelta )
 		(*DeviceIt)->Tick(TimeDelta);
 		(*DeviceIt)->SendControllerEvents();
 	}
+
+	FScopeLock Lock(&CriticalSection);
+	if(bOrientationChanged && Windows.Num() > 0)
+	{
+		int32 WindowX,WindowY, WindowWidth,WindowHeight;
+		Windows[0]->GetFullScreenInfo(WindowX, WindowY, WindowWidth, WindowHeight);
+
+		GenericApplication::GetMessageHandler()->OnSizeChanged(Windows[0],WindowWidth,WindowHeight, false);
+		GenericApplication::GetMessageHandler()->OnResizingWindow(Windows[0]);
+		FDisplayMetrics DisplayMetrics;
+		FDisplayMetrics::GetDisplayMetrics(DisplayMetrics);
+		BroadcastDisplayMetricsChanged(DisplayMetrics);
+		bOrientationChanged = false;
+	}
 }
 
 FPlatformRect FIOSApplication::GetWorkArea( const FPlatformRect& CurrentWindow ) const
@@ -97,3 +115,9 @@ TSharedRef< FGenericWindow > FIOSApplication::MakeWindow()
 {
 	return FIOSWindow::Make();
 }
+
+void FIOSApplication::OrientationChanged(UIDeviceOrientation orientation)
+{
+	FScopeLock Lock(&CriticalSection);
+	bOrientationChanged = true;
+}
diff --git a/Engine/Source/Runtime/Core/Private/IOS/IOSView.cpp b/Engine/Source/Runtime/Core/Private/IOS/IOSView.cpp
--- a/Engine/Source/Runtime/Core/Private/IOS/IOSView.cpp
+++ b/Engine/Source/Runtime/Core/Private/IOS/IOSView.cpp
@@ -3,6 +3,7 @@
 #include "CorePrivatePCH.h"
 #include "IOSView.h"
 #include "IOSAppDelegate.h"
+#include "IOSApplication.h"
 #include "IOS/IOSInputInterface.h"
 
 #include <OpenGLES/ES2/gl.h>
@@ -129,7 +130,7 @@ id<MTLDevice> GMetalDevice = nil;
 		SwapCount = 0;
 
 		FMemory::Memzero(AllTouches, sizeof(AllTouches));
-
+		[self setAutoresizingMask: UIViewAutoresizingFlexibleWidth | UIViewAutoresizingFlexibleHeight];
 		bIsInitialized = false;
 	}
 	return self;
@@ -216,6 +217,7 @@ id<MTLDevice> GMetalDevice = nil;
 			check(glGetError() == 0);
 			glFramebufferRenderbuffer(GL_FRAMEBUFFER, GL_COLOR_ATTACHMENT0, GL_RENDERBUFFER, OnScreenColorRenderBuffer);
 			check(glGetError() == 0);
+			check(glCheckFramebufferStatus(GL_FRAMEBUFFER) == GL_FRAMEBUFFER_COMPLETE);
 
 #if USE_DETAILED_IPHONE_MEM_TRACKING
 			//This value is used to allow the engine to track gl allocated memory see GetIPhoneOpenGLBackBufferSize
@@ -231,6 +233,47 @@ id<MTLDevice> GMetalDevice = nil;
 	return true;
 }
 
+/**
+ * If view is resized, update the frame buffer so it is the same size as the display area.
+ */
+- (void)layoutSubviews
+{
+	UIDeviceOrientation orientation = [[UIDevice currentDevice] orientation];
+	FIOSApplication::OrientationChanged(orientation);
+}
+
+-(void)UpdateRenderWidth:(uint32)Width andHeight:(uint32)Height
+{
+#if HAS_METAL
+	if (bIsUsingMetal)
+	{
+		if (MetalDevice != nil)
+		{
+			// grab the MetalLayer and typecast it to match what's in layerClass
+			CAMetalLayer* MetalLayer = (CAMetalLayer*)self.layer;
+			CGSize drawableSize = CGSizeMake(Width, Height);
+			check( drawableSize.width == self.bounds.size.width && drawableSize.height == self.bounds.size.height);
+			
+			drawableSize.width  *= self.contentScaleFactor;
+			drawableSize.height *= self.contentScaleFactor;
+			
+			MetalLayer.drawableSize = drawableSize;
+		}
+		return;
+	}
+#endif
+	check( ([EAGLContext currentContext] == Context));
+	// Allocate color buffer based on the current layer size
+	glBindRenderbuffer(GL_RENDERBUFFER, OnScreenColorRenderBuffer);
+	check(glGetError() == 0);
+	[Context renderbufferStorage:GL_RENDERBUFFER fromDrawable:(CAEAGLLayer*)self.layer];
+	glBindFramebuffer(GL_FRAMEBUFFER, ResolveFrameBuffer);
+	check(glGetError() == 0);
+	glFramebufferRenderbuffer(GL_FRAMEBUFFER, GL_COLOR_ATTACHMENT0, GL_RENDERBUFFER, OnScreenColorRenderBuffer);
+	check(glGetError() == 0);
+	check(glCheckFramebufferStatus(GL_FRAMEBUFFER) == GL_FRAMEBUFFER_COMPLETE);
+}
+
 #if HAS_METAL
 - (id<CAMetalDrawable>)MakeDrawable
 {
@@ -249,6 +292,7 @@ id<MTLDevice> GMetalDevice = nil;
 			if(ResolveFrameBuffer)
 			{
 				glDeleteFramebuffers(1, &ResolveFrameBuffer);
+				ResolveFrameBuffer = 0;
 			}
 			if(OnScreenColorRenderBuffer)
 			{
diff --git a/Engine/Source/Runtime/Core/Public/IOS/IOSApplication.h b/Engine/Source/Runtime/Core/Public/IOS/IOSApplication.h
--- a/Engine/Source/Runtime/Core/Public/IOS/IOSApplication.h
+++ b/Engine/Source/Runtime/Core/Public/IOS/IOSApplication.h
@@ -27,6 +27,8 @@ public:
 
 	virtual void AddExternalInputDevice(TSharedPtr<class IInputDevice> InputDevice);
 
+	static void OrientationChanged(UIDeviceOrientation orientation);
+
 protected:
 	virtual void InitializeWindow( const TSharedRef< FGenericWindow >& Window, const TSharedRef< FGenericWindowDefinition >& InDefinition, const TSharedPtr< FGenericWindow >& InParent, const bool bShowImmediately ) override;
 
@@ -44,4 +46,7 @@ private:
 	bool bHasLoadedInputPlugins;
 
 	TArray< TSharedRef< FIOSWindow > > Windows;
+
+	static FCriticalSection CriticalSection;
+	static bool bOrientationChanged;
 };
diff --git a/Engine/Source/Runtime/Core/Public/IOS/IOSView.h b/Engine/Source/Runtime/Core/Public/IOS/IOSView.h
--- a/Engine/Source/Runtime/Core/Public/IOS/IOSView.h
+++ b/Engine/Source/Runtime/Core/Public/IOS/IOSView.h
@@ -48,6 +48,7 @@
 
 -(bool)CreateFramebuffer:(bool)bIsForOnDevice;
 -(void)DestroyFramebuffer;
+-(void)UpdateRenderWidth:(uint32)Width andHeight:(uint32)Height;
 
 //// GL FUNCTIONALITY
 @property (nonatomic) GLuint OnScreenColorRenderBuffer;
diff --git a/Engine/Source/Runtime/IOS/MetalRHI/Private/MetalCommands.cpp b/Engine/Source/Runtime/IOS/MetalRHI/Private/MetalCommands.cpp
--- a/Engine/Source/Runtime/IOS/MetalRHI/Private/MetalCommands.cpp
+++ b/Engine/Source/Runtime/IOS/MetalRHI/Private/MetalCommands.cpp
@@ -89,6 +89,13 @@ void FMetalDynamicRHI::RHIDispatchIndirectComputeShader(FVertexBufferRHIParamRef
 
 void FMetalDynamicRHI::RHISetViewport(uint32 MinX,uint32 MinY,float MinZ,uint32 MaxX,uint32 MaxY,float MaxZ)
 {
+	// Ensure viewport fits inside the currently bound RT.
+	FIntPoint TargetDims = FMetalManager::Get()->GetBoundRenderTargetDimensions();
+	MinX = FMath::Min(MinX, (uint32)TargetDims.X);
+	MaxX = FMath::Min(MaxX, (uint32)TargetDims.X);
+	MinY = FMath::Min(MinY, (uint32)TargetDims.Y);
+	MaxY = FMath::Min(MaxY, (uint32)TargetDims.Y);
+	
 	MTLViewport Viewport;
 	Viewport.originX = MinX;
 	Viewport.originY = MinY;
diff --git a/Engine/Source/Runtime/IOS/MetalRHI/Private/MetalManager.cpp b/Engine/Source/Runtime/IOS/MetalRHI/Private/MetalManager.cpp
--- a/Engine/Source/Runtime/IOS/MetalRHI/Private/MetalManager.cpp
+++ b/Engine/Source/Runtime/IOS/MetalRHI/Private/MetalManager.cpp
@@ -221,6 +221,8 @@ FMetalManager::FMetalManager()
 		FIOSPlatformRHIFramePacer::InitWithEvent( FrameReadyEvent );
 	}
 	
+	BoundRenderTargetDimensions = FIntPoint(0,0);
+	
 	InitFrame();
 }
 
@@ -595,9 +597,47 @@ void FMetalManager::SetRenderTargetsInfo(const FRHISetRenderTargetsInfo& RenderT
 		return;
 	}
 
+	// commit pending commands on the old render target
+	if (CurrentContext)
+	{
+#if UE_BUILD_DEBUG || UE_BUILD_DEVELOPMENT
+		//		if (NumDrawCalls == 0)
+		//		{
+		//			NSLog(@"There were %d draw calls for an RT in frame %lld", NumDrawCalls, GFrameCounter);
+		//		}
+#endif
+		
+		[CurrentContext endEncoding];
+		NumDrawCalls = 0;
+		
+		[CurrentContext release];
+		
+		// commit the buffer for this context
+		[CurrentCommandBuffer commit];
+		UNTRACK_OBJECT(CurrentCommandBuffer);
+		[CurrentCommandBuffer release];
+
+		// Check if CurrentCommandBuffer was rendering to the BackBuffer.
+		if( PreviousRenderTargetsInfo.NumColorRenderTargets == 1 )
+		{
+			const FRHIRenderTargetView& RenderTargetView = PreviousRenderTargetsInfo.ColorRenderTarget[0];
+			FMetalSurface& Surface = *GetMetalSurfaceFromRHITexture(RenderTargetView.Texture);
+			if(&Surface == &BackBuffer->Surface && CurrentDrawable != nil)
+			{
+				// release our record of it to ensure we are allocated a new drawable.
+				CurrentDrawable = nil;
+				BackBuffer->Surface.Texture = nil;
+			}
+		}
+		
+		// create the command buffer for this frame
+		CreateCurrentCommandBuffer(false);
+	}
+	
 	// back this up for next frame
 	PreviousRenderTargetsInfo = RenderTargetsInfo;
 
+	FIntPoint MaxDimensions(0,0);
 	// at this point, we need to fully set up an encoder/command buffer, so make a new one (autoreleased)
 	MTLRenderPassDescriptor* RenderPass = [MTLRenderPassDescriptor renderPassDescriptor];
 
@@ -619,6 +659,9 @@ void FMetalManager::SetRenderTargetsInfo(const FRHISetRenderTargetsInfo& RenderT
 			FMetalSurface& Surface = *GetMetalSurfaceFromRHITexture(RenderTargetView.Texture);
 			FormatKey = Surface.FormatKey;
 		
+			MaxDimensions.X = FMath::Max((uint32)MaxDimensions.X, Surface.SizeX);
+			MaxDimensions.Y = FMath::Max((uint32)MaxDimensions.Y, Surface.SizeY);
+
 			// if this is the back buffer, make sure we have a usable drawable
 			ConditionalUpdateBackBuffer(Surface);
             
@@ -685,6 +728,9 @@ void FMetalManager::SetRenderTargetsInfo(const FRHISetRenderTargetsInfo& RenderT
 	if (RenderTargetsInfo.DepthStencilRenderTarget.Texture != nullptr)
 	{
 		FMetalSurface& Surface= *GetMetalSurfaceFromRHITexture(RenderTargetsInfo.DepthStencilRenderTarget.Texture);
+		MaxDimensions.X = FMath::Max((uint32)MaxDimensions.X, Surface.SizeX);
+		MaxDimensions.Y = FMath::Max((uint32)MaxDimensions.Y, Surface.SizeY);
+
 		if (Surface.Texture != nil)
 		{
 			MTLRenderPassDepthAttachmentDescriptor* DepthAttachment = [[MTLRenderPassDepthAttachmentDescriptor alloc] init];
@@ -728,35 +774,13 @@ void FMetalManager::SetRenderTargetsInfo(const FRHISetRenderTargetsInfo& RenderT
 		}
 	}
 
+	BoundRenderTargetDimensions = MaxDimensions;
+
 	// update hash for the depth buffer
 	SET_HASH(OFFSET_DEPTH_ENABLED, NUMBITS_DEPTH_ENABLED, (Pipeline.DepthTargetFormat == MTLPixelFormatInvalid ? 0 : 1));
 	SET_HASH(OFFSET_STENCIL_ENABLED, NUMBITS_STENCIL_ENABLED, (Pipeline.StencilTargetFormat == MTLPixelFormatInvalid ? 0 : 1));
 	SET_HASH(OFFSET_SAMPLE_COUNT, NUMBITS_SAMPLE_COUNT, Pipeline.SampleCount);
 
-	// commit pending commands on the old render target
-	if (CurrentContext)
-	{
-#if UE_BUILD_DEBUG || UE_BUILD_DEVELOPMENT
-//		if (NumDrawCalls == 0)
-//		{
-//			NSLog(@"There were %d draw calls for an RT in frame %lld", NumDrawCalls, GFrameCounter);
-//		}
-#endif
-
-		[CurrentContext endEncoding];
-		NumDrawCalls = 0;
-
-		[CurrentContext release];
-
-		// commit the buffer for this context
-		[CurrentCommandBuffer commit];
-		UNTRACK_OBJECT(CurrentCommandBuffer);
-		[CurrentCommandBuffer release];
-
-		// create the command buffer for this frame
-		CreateCurrentCommandBuffer(false);
-	}
-
 	// make a new render context to use to render to the framebuffer
 	CurrentContext = [CurrentCommandBuffer renderCommandEncoderWithDescriptor:RenderPass];
 	[CurrentContext retain];
@@ -1013,3 +1037,21 @@ void FMetalManager::Dispatch(uint32 ThreadGroupCountX, uint32 ThreadGroupCountY,
 	//@todo-rco: setThreadgroupMemoryLength?
 	[CurrentComputeContext dispatchThreadgroups:Threadgroups threadsPerThreadgroup:ThreadgroupCounts];
 }
+
+void FMetalManager::ResizeBackBuffer(uint32 InSizeX, uint32 InSizeY)
+{
+	IOSAppDelegate* AppDelegate = [IOSAppDelegate GetDelegate];
+	FIOSView* GLView = AppDelegate.IOSView;
+	[GLView UpdateRenderWidth:InSizeX andHeight:InSizeY];
+	FRHIResourceCreateInfo CreateInfo;
+	BackBuffer = (FMetalTexture2D*)(FTexture2DRHIParamRef)RHICreateTexture2D(InSizeX, InSizeY, PF_B8G8R8A8, 1, 1, TexCreate_RenderTargetable | TexCreate_Presentable, CreateInfo);
+
+	// ensure a new drawable is created when this new backbuffer is used:
+	PreviousRenderTargetsInfo.NumColorRenderTargets = 0;
+	CurrentDrawable = nil;
+	BackBuffer->Surface.Texture = nil;
+	
+	// get the size of the window
+	CGRect ViewFrame = [[IOSAppDelegate GetDelegate].IOSView frame];
+	check(ViewFrame.size.width == InSizeX && ViewFrame.size.height == InSizeY);
+}
\ No newline at end of file
diff --git a/Engine/Source/Runtime/IOS/MetalRHI/Private/MetalManager.h b/Engine/Source/Runtime/IOS/MetalRHI/Private/MetalManager.h
--- a/Engine/Source/Runtime/IOS/MetalRHI/Private/MetalManager.h
+++ b/Engine/Source/Runtime/IOS/MetalRHI/Private/MetalManager.h
@@ -163,6 +163,9 @@ public:
 	void SetComputeShader(FMetalComputeShader* InComputeShader);
 	void Dispatch(uint32 ThreadGroupCountX, uint32 ThreadGroupCountY, uint32 ThreadGroupCountZ);
 
+	void ResizeBackBuffer(uint32 InSizeX, uint32 InSizeY);
+	FIntPoint GetBoundRenderTargetDimensions() { return BoundRenderTargetDimensions; }
+
 protected:
 	FMetalManager();
 	void InitFrame();
@@ -285,6 +288,9 @@ protected:
 
 	void CommitNonComputeShaderConstants();
 
+	/** The dimensions of the currently bound RT */
+	FIntPoint BoundRenderTargetDimensions;
+
 private:
 	
 	/** Event for coordinating pausing of render thread to keep inline with the ios display link. */
diff --git a/Engine/Source/Runtime/IOS/MetalRHI/Private/MetalViewport.cpp b/Engine/Source/Runtime/IOS/MetalRHI/Private/MetalViewport.cpp
--- a/Engine/Source/Runtime/IOS/MetalRHI/Private/MetalViewport.cpp
+++ b/Engine/Source/Runtime/IOS/MetalRHI/Private/MetalViewport.cpp
@@ -16,6 +16,10 @@ FMetalViewport::~FMetalViewport()
 {
 }
 
+void FMetalViewport::Resize(uint32 InSizeX,uint32 InSizeY,bool bInIsFullscreen)
+{
+	FMetalManager::Get()->ResizeBackBuffer(InSizeX, InSizeY);
+}
 
 /*=============================================================================
  *	The following RHI functions must be called from the main thread.
@@ -31,6 +35,7 @@ void FMetalDynamicRHI::RHIResizeViewport(FViewportRHIParamRef ViewportRHI,uint32
 	check( IsInGameThread() );
 
 	DYNAMIC_CAST_METALRESOURCE(Viewport,Viewport);
+	Viewport->Resize(SizeX,SizeY,bIsFullscreen);
 }
 
 void FMetalDynamicRHI::RHITick( float DeltaTime )
diff --git a/Engine/Source/Runtime/IOS/MetalRHI/Public/MetalViewport.h b/Engine/Source/Runtime/IOS/MetalRHI/Public/MetalViewport.h
--- a/Engine/Source/Runtime/IOS/MetalRHI/Public/MetalViewport.h
+++ b/Engine/Source/Runtime/IOS/MetalRHI/Public/MetalViewport.h
@@ -13,5 +13,6 @@ public:
 	FMetalViewport(void* WindowHandle, uint32 InSizeX,uint32 InSizeY,bool bInIsFullscreen);
 	~FMetalViewport();
 
+	void Resize(uint32 InSizeX, uint32 InSizeY, bool bInIsFullscreen);
 };
 
diff --git a/Engine/Source/Runtime/OpenGLDrv/Private/IOS/IOSOpenGL.cpp b/Engine/Source/Runtime/OpenGLDrv/Private/IOS/IOSOpenGL.cpp
--- a/Engine/Source/Runtime/OpenGLDrv/Private/IOS/IOSOpenGL.cpp
+++ b/Engine/Source/Runtime/OpenGLDrv/Private/IOS/IOSOpenGL.cpp
@@ -227,8 +227,9 @@ FRHITexture* PlatformCreateBuiltinBackBuffer(FOpenGLDynamicRHI* OpenGLRHI, uint3
 {
 	IOSAppDelegate* AppDelegate = [IOSAppDelegate GetDelegate];
 	FIOSView* GLView = AppDelegate.IOSView;
+	[GLView UpdateRenderWidth:SizeX andHeight:SizeY];
 
-	uint32 Flags = TexCreate_RenderTargetable;
+	uint32 Flags = TexCreate_RenderTargetable | TexCreate_Presentable;
 	FOpenGLTexture2D* Texture2D = new FOpenGLTexture2D(OpenGLRHI, GLView.OnScreenColorRenderBuffer, GL_RENDERBUFFER, GL_COLOR_ATTACHMENT0, SizeX, SizeY, 0, 1, 1, 1, PF_B8G8R8A8, false, false, Flags, nullptr);
 	OpenGLTextureAllocated(Texture2D, Flags);
 
@@ -241,12 +242,21 @@ void PlatformResizeGLContext( FPlatformOpenGLDevice* Device, FPlatformOpenGLCont
 	VERIFY_GL_SCOPE();
 	glBindFramebuffer(GL_FRAMEBUFFER, Context->ViewportFramebuffer);
 	glFramebufferRenderbuffer(GL_FRAMEBUFFER, GL_COLOR_ATTACHMENT0, GL_RENDERBUFFER, BackBufferResource);
+	check (glCheckFramebufferStatus(GL_FRAMEBUFFER) == GL_FRAMEBUFFER_COMPLETE);
 
 	glViewport(0, 0, SizeX, SizeY);
 
 	// @todo-mobile Do we need to clear here?
 	glClearColor(0, 0, 0, 0);
 	glClear(GL_COLOR_BUFFER_BIT);
+
+	#if UE_BUILD_DEBUG
+	// check size requested and renderbuffer size agree.
+	GLint RenderBufferWidth, RenderBufferHeight;
+	glGetRenderbufferParameteriv(GL_RENDERBUFFER, GL_RENDERBUFFER_WIDTH, &RenderBufferWidth);
+	glGetRenderbufferParameteriv(GL_RENDERBUFFER, GL_RENDERBUFFER_HEIGHT, &RenderBufferHeight);
+	check(SizeX == RenderBufferWidth && SizeY == RenderBufferHeight);
+	#endif
 }
 
 void PlatformGetSupportedResolution(uint32 &Width, uint32 &Height)
diff --git a/Engine/Source/Runtime/OpenGLDrv/Private/OpenGLCommands.cpp b/Engine/Source/Runtime/OpenGLDrv/Private/OpenGLCommands.cpp
--- a/Engine/Source/Runtime/OpenGLDrv/Private/OpenGLCommands.cpp
+++ b/Engine/Source/Runtime/OpenGLDrv/Private/OpenGLCommands.cpp
@@ -1678,31 +1678,33 @@ void FOpenGLDynamicRHI::RHISetRenderTargets(
 		// @todo-mobile
 
 		FOpenGLContextState& ContextState = GetContextStateForCurrentContext();
-		GLuint NewColorRT = PendingState.RenderTargets[0] ? PendingState.RenderTargets[0]->Resource : 0;
-		// If the color buffer did not change and we are disabling depth, do not switch depth and assume 
+		GLuint NewColorRTResource = PendingState.RenderTargets[0] ? PendingState.RenderTargets[0]->Resource : 0;
+		GLenum NewColorTargetType = PendingState.RenderTargets[0] ? PendingState.RenderTargets[0]->Target : 0;
+		// If the color buffer did not change and we are disabling depth, do not switch depth and assume
 		// the high level will disable depth test/write (so we can avoid a logical buffer store);
 		// if both are set to nothing, then it's an endframe so we don't want to switch either...
 		if (NewDepthStencilRT == NULL && PendingState.DepthStencil != NULL)
 		{
+			const bool bColorBufferUnchanged = ContextState.LastES2ColorRTResource == NewColorRTResource && ContextState.LastES2ColorTargetType == NewColorTargetType;
 #if PLATFORM_ANDROID
 			//color RT being 0 means backbuffer is being used. Hence taking only comparison with previous RT into consideration. Fixes black screen issue.
-			if ( ContextState.LastES2ColorRT == NewColorRT)
+			if (bColorBufferUnchanged)
 #else
-			if (NewColorRT == 0 || ContextState.LastES2ColorRT == NewColorRT)
+			if (NewColorRTResource == 0 || bColorBufferUnchanged)
 #endif
 			{
 				return;
 			}
 			else
 			{
-				ContextState.LastES2ColorRT = NewColorRT;
-				ContextState.LastES2DepthRT = NewDepthStencilRT ? NewDepthStencilRT->Resource : 0;
+				ContextState.LastES2ColorRTResource = NewColorRTResource;
+				ContextState.LastES2ColorTargetType = NewColorTargetType;
 			}
 		}
 		else
 		{
-			ContextState.LastES2ColorRT = NewColorRT;
-			ContextState.LastES2DepthRT = NewDepthStencilRT ? NewDepthStencilRT->Resource : 0;
+				ContextState.LastES2ColorRTResource = NewColorRTResource;
+				ContextState.LastES2ColorTargetType = NewColorTargetType;
 		}
 	}
 	PendingState.DepthStencil = NewDepthStencilRT;
diff --git a/Engine/Source/Runtime/OpenGLDrv/Private/OpenGLDevice.cpp b/Engine/Source/Runtime/OpenGLDrv/Private/OpenGLDevice.cpp
--- a/Engine/Source/Runtime/OpenGLDrv/Private/OpenGLDevice.cpp
+++ b/Engine/Source/Runtime/OpenGLDrv/Private/OpenGLDevice.cpp
@@ -145,8 +145,7 @@ void FOpenGLDynamicRHI::RHIBeginFrame()
 #if PLATFORM_ANDROID //adding #if since not sure if this is required for any other platform.
 	//we need to differential between 0 (backbuffer) and lastcolorRT.
 	FOpenGLContextState& ContextState = GetContextStateForCurrentContext();
-	ContextState.LastES2ColorRT = 0xFFFFFFFF;
-	ContextState.LastES2DepthRT = 0xFFFFFFFF;
+	ContextState.LastES2ColorRTResource = 0xFFFFFFFF;
 	PendingState.DepthStencil = 0 ;
 #endif
 }
diff --git a/Engine/Source/Runtime/OpenGLDrv/Private/OpenGLTexture.cpp b/Engine/Source/Runtime/OpenGLDrv/Private/OpenGLTexture.cpp
--- a/Engine/Source/Runtime/OpenGLDrv/Private/OpenGLTexture.cpp
+++ b/Engine/Source/Runtime/OpenGLDrv/Private/OpenGLTexture.cpp
@@ -1405,8 +1405,32 @@ TOpenGLTexture<RHIResourceType>::~TOpenGLTexture()
 
 		if( Resource != 0 )
 		{
-			OpenGLRHI->InvalidateTextureResourceInCache( Resource );
-			FOpenGL::DeleteTextures( 1, &Resource );
+			switch(Target)
+			{
+				case GL_TEXTURE_2D:
+ 				case GL_TEXTURE_2D_MULTISAMPLE:
+				case GL_TEXTURE_3D:
+				case GL_TEXTURE_CUBE_MAP:
+				case GL_TEXTURE_2D_ARRAY:
+				case GL_TEXTURE_CUBE_MAP_ARRAY:
+				{
+					OpenGLRHI->InvalidateTextureResourceInCache( Resource );
+					FOpenGL::DeleteTextures( 1, &Resource );
+					break;
+				}
+				case GL_RENDERBUFFER:
+				{
+					if (!(this->GetFlags() & TexCreate_Presentable))
+					{
+						glDeleteRenderbuffers(1, &Resource);
+					}
+					break;
+				}
+				default:
+				{
+					checkNoEntry();
+				}
+			}
 		}
 		
 		if(TextureRange)
diff --git a/Engine/Source/Runtime/OpenGLDrv/Public/OpenGLState.h b/Engine/Source/Runtime/OpenGLDrv/Public/OpenGLState.h
--- a/Engine/Source/Runtime/OpenGLDrv/Public/OpenGLState.h
+++ b/Engine/Source/Runtime/OpenGLDrv/Public/OpenGLState.h
@@ -314,8 +314,8 @@ struct FOpenGLContextState : public FOpenGLCommonState
 	float							ClearDepth;
 
 	// @todo-mobile: Used to cache the last color attachment to optimize logical buffer loads
-	GLuint							LastES2ColorRT;
-	GLuint							LastES2DepthRT;
+	GLuint							LastES2ColorRTResource;
+	GLenum							LastES2ColorTargetType;
 
 	FOpenGLCachedAttr				VertexAttrs[NUM_OPENGL_VERTEX_STREAMS];
 	FOpenGLStream					VertexStreams[NUM_OPENGL_VERTEX_STREAMS];
@@ -342,12 +342,11 @@ struct FOpenGLContextState : public FOpenGLCommonState
 	,	ClearStencil(0xFFFF)
 	,	ClearDepth(-1.0f)
 #if PLATFORM_ANDROID
-	,	LastES2ColorRT(0xFFFFFFFF)
-	,	LastES2DepthRT(0xFFFFFFFF)
+	,	LastES2ColorRTResource(0xFFFFFFFF)
 #else
-	,	LastES2ColorRT(0)
-	,	LastES2DepthRT(0)
+	,	LastES2ColorRTResource(0)
 #endif
+	,	LastES2ColorTargetType(GL_NONE)
 	, VertexDecl(0)
 	, ActiveAttribMask(0)
 	, ActiveStreamMask(0)
