Use  the minimum possible dimensions for  BoundRenderTargetDimensions.
Stops Metal validation getting upset with viewport dimensions  during rotation (when mobilehdr = false)

diff --git a/Engine/Source/Runtime/IOS/MetalRHI/Private/MetalManager.cpp b/Engine/Source/Runtime/IOS/MetalRHI/Private/MetalManager.cpp
--- a/Engine/Source/Runtime/IOS/MetalRHI/Private/MetalManager.cpp
+++ b/Engine/Source/Runtime/IOS/MetalRHI/Private/MetalManager.cpp
@@ -637,7 +637,7 @@ void FMetalManager::SetRenderTargetsInfo(const FRHISetRenderTargetsInfo& RenderT
 	// back this up for next frame
 	PreviousRenderTargetsInfo = RenderTargetsInfo;
 
-	FIntPoint MaxDimensions(0,0);
+	FIntPoint MaxDimensions(TNumericLimits<decltype(FIntPoint::X)>::Max(),TNumericLimits<decltype(FIntPoint::Y)>::Max());
 	// at this point, we need to fully set up an encoder/command buffer, so make a new one (autoreleased)
 	MTLRenderPassDescriptor* RenderPass = [MTLRenderPassDescriptor renderPassDescriptor];
 
@@ -659,8 +659,8 @@ void FMetalManager::SetRenderTargetsInfo(const FRHISetRenderTargetsInfo& RenderT
 			FMetalSurface& Surface = *GetMetalSurfaceFromRHITexture(RenderTargetView.Texture);
 			FormatKey = Surface.FormatKey;
 		
-			MaxDimensions.X = FMath::Max((uint32)MaxDimensions.X, Surface.SizeX);
-			MaxDimensions.Y = FMath::Max((uint32)MaxDimensions.Y, Surface.SizeY);
+			MaxDimensions.X = FMath::Min((uint32)MaxDimensions.X, Surface.SizeX);
+			MaxDimensions.Y = FMath::Min((uint32)MaxDimensions.Y, Surface.SizeY);
 
 			// if this is the back buffer, make sure we have a usable drawable
 			ConditionalUpdateBackBuffer(Surface);
@@ -728,8 +728,8 @@ void FMetalManager::SetRenderTargetsInfo(const FRHISetRenderTargetsInfo& RenderT
 	if (RenderTargetsInfo.DepthStencilRenderTarget.Texture != nullptr)
 	{
 		FMetalSurface& Surface= *GetMetalSurfaceFromRHITexture(RenderTargetsInfo.DepthStencilRenderTarget.Texture);
-		MaxDimensions.X = FMath::Max((uint32)MaxDimensions.X, Surface.SizeX);
-		MaxDimensions.Y = FMath::Max((uint32)MaxDimensions.Y, Surface.SizeY);
+		MaxDimensions.X = FMath::Min((uint32)MaxDimensions.X, Surface.SizeX);
+		MaxDimensions.Y = FMath::Min((uint32)MaxDimensions.Y, Surface.SizeY);
 
 		if (Surface.Texture != nil)
 		{
@@ -774,6 +774,9 @@ void FMetalManager::SetRenderTargetsInfo(const FRHISetRenderTargetsInfo& RenderT
 		}
 	}
 
+	check( MaxDimensions.X != TNumericLimits<decltype(FIntPoint::X)>::Max() &&
+		   MaxDimensions.Y != TNumericLimits<decltype(FIntPoint::Y)>::Max());
+	
 	BoundRenderTargetDimensions = MaxDimensions;
 
 	// update hash for the depth buffer
