UE4 - fix ResourceCast problems for metal

diff --git a/Engine/Source/Runtime/IOS/MetalRHI/Private/MetalTexture.cpp b/Engine/Source/Runtime/IOS/MetalRHI/Private/MetalTexture.cpp
--- a/Engine/Source/Runtime/IOS/MetalRHI/Private/MetalTexture.cpp
+++ b/Engine/Source/Runtime/IOS/MetalRHI/Private/MetalTexture.cpp
@@ -545,7 +545,7 @@ void FMetalDynamicRHI::RHIUnlockTexture2DArray(FTexture2DArrayRHIParamRef Textur
 
 void FMetalDynamicRHI::RHIUpdateTexture2D(FTexture2DRHIParamRef TextureRHI, uint32 MipIndex, const struct FUpdateTextureRegion2D& UpdateRegion, uint32 SourcePitch, const uint8* SourceData)
 {	
-	FMetalTexture3D* Texture = ResourceCast(TextureRHI);
+	FMetalTexture2D* Texture = ResourceCast(TextureRHI);
 
 }
 
diff --git a/Engine/Source/Runtime/IOS/MetalRHI/Public/MetalResources.h b/Engine/Source/Runtime/IOS/MetalRHI/Public/MetalResources.h
--- a/Engine/Source/Runtime/IOS/MetalRHI/Public/MetalResources.h
+++ b/Engine/Source/Runtime/IOS/MetalRHI/Public/MetalResources.h
@@ -611,8 +611,3 @@ struct TMetalResourceTraits<FRHIBlendState>
 {
 	typedef FMetalBlendState TConcreteType;
 };
-template<>
-struct TMetalResourceTraits<FRHIViewport>
-{
-	typedef FMetalViewport TConcreteType;
-};
diff --git a/Engine/Source/Runtime/IOS/MetalRHI/Public/MetalViewport.h b/Engine/Source/Runtime/IOS/MetalRHI/Public/MetalViewport.h
--- a/Engine/Source/Runtime/IOS/MetalRHI/Public/MetalViewport.h
+++ b/Engine/Source/Runtime/IOS/MetalRHI/Public/MetalViewport.h
@@ -16,3 +16,9 @@ public:
 	void Resize(uint32 InSizeX, uint32 InSizeY, bool bInIsFullscreen);
 };
 
+template<>
+struct TMetalResourceTraits<FRHIViewport>
+{
+	typedef FMetalViewport TConcreteType;
+};
+
