UE4 - fix Metal CIS error

diff --git a/Engine/Source/Runtime/IOS/MetalRHI/Private/MetalRenderTarget.cpp b/Engine/Source/Runtime/IOS/MetalRHI/Private/MetalRenderTarget.cpp
--- a/Engine/Source/Runtime/IOS/MetalRHI/Private/MetalRenderTarget.cpp
+++ b/Engine/Source/Runtime/IOS/MetalRHI/Private/MetalRenderTarget.cpp
@@ -33,8 +33,6 @@ void FMetalDynamicRHI::RHIReadSurfaceFloatData(FTextureRHIParamRef TextureRHI, F
 	//kick the current command buffer.
 	FMetalManager::Get()->SubmitCommandBufferAndWait();
 
-	FMetalTexture2D* Texture = ResourceCast(TextureRHI);
-	
 	FMetalSurface* Surface = GetMetalSurfaceFromRHITexture(TextureRHI);
 	
 	// verify the input image format (but don't crash)
