Fix ios/metal issues when mobilescaling is != 1.
Fix for  metal''''''s incorrectly set backbuffer size.
Factor scaling size into some asserts.
removed erroneous GL context check.

diff --git a/Engine/Source/Runtime/Core/Private/IOS/IOSView.cpp b/Engine/Source/Runtime/Core/Private/IOS/IOSView.cpp
--- a/Engine/Source/Runtime/Core/Private/IOS/IOSView.cpp
+++ b/Engine/Source/Runtime/Core/Private/IOS/IOSView.cpp
@@ -252,17 +252,14 @@ id<MTLDevice> GMetalDevice = nil;
 			// grab the MetalLayer and typecast it to match what's in layerClass
 			CAMetalLayer* MetalLayer = (CAMetalLayer*)self.layer;
 			CGSize drawableSize = CGSizeMake(Width, Height);
-			check( drawableSize.width == self.bounds.size.width && drawableSize.height == self.bounds.size.height);
-			
-			drawableSize.width  *= self.contentScaleFactor;
-			drawableSize.height *= self.contentScaleFactor;
-			
+			check( drawableSize.width == (self.bounds.size.width * self.contentScaleFactor) && drawableSize.height == (self.bounds.size.height * self.contentScaleFactor));
+
 			MetalLayer.drawableSize = drawableSize;
 		}
 		return;
 	}
 #endif
-	check( ([EAGLContext currentContext] == Context));
+
 	// Allocate color buffer based on the current layer size
 	glBindRenderbuffer(GL_RENDERBUFFER, OnScreenColorRenderBuffer);
 	check(glGetError() == 0);
diff --git a/Engine/Source/Runtime/IOS/MetalRHI/Private/MetalManager.cpp b/Engine/Source/Runtime/IOS/MetalRHI/Private/MetalManager.cpp
--- a/Engine/Source/Runtime/IOS/MetalRHI/Private/MetalManager.cpp
+++ b/Engine/Source/Runtime/IOS/MetalRHI/Private/MetalManager.cpp
@@ -200,7 +200,9 @@ FMetalManager::FMetalManager()
 	// get the size of the window
 	CGRect ViewFrame = [[IOSAppDelegate GetDelegate].IOSView frame];
 	FRHIResourceCreateInfo CreateInfo;
-	BackBuffer = (FMetalTexture2D*)(FTexture2DRHIParamRef)RHICreateTexture2D(ViewFrame.size.width, ViewFrame.size.height, PF_B8G8R8A8, 1, 1, TexCreate_RenderTargetable | TexCreate_Presentable, CreateInfo);
+
+	float ScalingFactor = [[IOSAppDelegate GetDelegate].IOSView contentScaleFactor];
+	BackBuffer = (FMetalTexture2D*)(FTexture2DRHIParamRef)RHICreateTexture2D(ViewFrame.size.width*ScalingFactor, ViewFrame.size.height*ScalingFactor, PF_B8G8R8A8, 1, 1, TexCreate_RenderTargetable | TexCreate_Presentable, CreateInfo);
 
 //@todo-rco: What Size???
 	// make a buffer for each shader type
@@ -1054,7 +1056,8 @@ void FMetalManager::ResizeBackBuffer(uint32 InSizeX, uint32 InSizeY)
 	CurrentDrawable = nil;
 	BackBuffer->Surface.Texture = nil;
 	
-	// get the size of the window
+	// check the size of the window
+	float ScalingFactor = [[IOSAppDelegate GetDelegate].IOSView contentScaleFactor];
 	CGRect ViewFrame = [[IOSAppDelegate GetDelegate].IOSView frame];
-	check(ViewFrame.size.width == InSizeX && ViewFrame.size.height == InSizeY);
+	check(ScalingFactor * ViewFrame.size.width == InSizeX && ScalingFactor * ViewFrame.size.height == InSizeY);
 }
\ No newline at end of file
