Providing an example of how to use the new Config Editor

diff --git a/Engine/Source/Runtime/IOS/IOSRuntimeSettings/Classes/IOSRuntimeSettings.h b/Engine/Source/Runtime/IOS/IOSRuntimeSettings/Classes/IOSRuntimeSettings.h
--- a/Engine/Source/Runtime/IOS/IOSRuntimeSettings/Classes/IOSRuntimeSettings.h
+++ b/Engine/Source/Runtime/IOS/IOSRuntimeSettings/Classes/IOSRuntimeSettings.h
@@ -93,7 +93,7 @@ public:
 	GENERATED_UCLASS_BODY()
 
 	// Should Game Center support (iOS Online Subsystem) be enabled?
-	UPROPERTY(GlobalConfig, EditAnywhere, Category = Online)
+	UPROPERTY(GlobalConfig, EditAnywhere, Category = Online, meta = (ConfigHierarchyEditable))
 	uint32 bEnableGameCenterSupport : 1;
 	
 	// Whether or not to add support for Metal API (requires IOS8 and A7 processors).
@@ -133,7 +133,7 @@ public:
 	bool bShipForArmV7S;
 	
 	// The name or ip address of the remote mac which will be used to build IOS
-	UPROPERTY(GlobalConfig, EditAnywhere, Category = "Build")
+	UPROPERTY(GlobalConfig, EditAnywhere, Category = "Build", meta = (ConfigHierarchyEditable))
 	FString RemoteServerName;
 
 	// Enable the use of RSync for remote builds on a mac
@@ -185,7 +185,7 @@ public:
 	FString VersionInfo;
     
     /** Set the maximum frame rate to save on power consumption */
-    UPROPERTY(GlobalConfig, EditAnywhere, Category = PowerUsage)
+	UPROPERTY(GlobalConfig, EditAnywhere, Category = PowerUsage, meta = (ConfigHierarchyEditable))
     TEnumAsByte<EPowerUsageFrameRateLock> FrameRateLock;
 
 	// Minimum iOS version this game supports
