Add RHISubmitCommandsHint to FMetalDynamicRHI

diff --git a/Engine/Source/Runtime/IOS/MetalRHI/Private/MetalCommands.cpp b/Engine/Source/Runtime/IOS/MetalRHI/Private/MetalCommands.cpp
--- a/Engine/Source/Runtime/IOS/MetalRHI/Private/MetalCommands.cpp
+++ b/Engine/Source/Runtime/IOS/MetalRHI/Private/MetalCommands.cpp
@@ -900,6 +900,11 @@ void FMetalDynamicRHI::RHIEnableDepthBoundsTest(bool bEnable, float MinDepth, fl
 	// not supported
 }
 
+void FMetalDynamicRHI::RHISubmitCommandsHint()
+{
+    
+}
+
 IRHICommandContext* FMetalDynamicRHI::RHIGetDefaultContext()
 {
 	return this;
