fix for various Strategy Game crashes when running on Metal
UE-10079
UE-13790
#ios
#metal

diff --git a/Engine/Source/Runtime/Engine/Private/MeshParticleVertexFactory.cpp b/Engine/Source/Runtime/Engine/Private/MeshParticleVertexFactory.cpp
--- a/Engine/Source/Runtime/Engine/Private/MeshParticleVertexFactory.cpp
+++ b/Engine/Source/Runtime/Engine/Private/MeshParticleVertexFactory.cpp
@@ -42,7 +42,7 @@ public:
 
 	virtual void SetMesh(FRHICommandList& RHICmdList, FShader* Shader,const FVertexFactory* VertexFactory,const FSceneView& View,const FMeshBatchElement& BatchElement,uint32 DataFlags) const override
 	{
-		const bool bInstanced = View.GetFeatureLevel() >= ERHIFeatureLevel::ES3_1;
+		const bool bInstanced = RHISupportsInstancing(GetFeatureLevelShaderPlatform(View.GetFeatureLevel()));
 		FMeshParticleVertexFactory* MeshParticleVF = (FMeshParticleVertexFactory*)VertexFactory;
 		FVertexShaderRHIParamRef VertexShaderRHI = Shader->GetVertexShader();
 		SetUniformBufferParameter(RHICmdList, VertexShaderRHI, Shader->GetUniformBufferParameter<FMeshParticleUniformParameters>(), MeshParticleVF->GetUniformBuffer() );
@@ -86,7 +86,7 @@ void FMeshParticleVertexFactory::InitRHI()
 {
 	FVertexDeclarationElementList Elements;
 
-	const bool bInstanced = GetFeatureLevel() >= ERHIFeatureLevel::ES3_1;
+    const bool bInstanced = RHISupportsInstancing(GetFeatureLevelShaderPlatform(GetFeatureLevel()));
 
 	if (Data.bInitialized)
 	{
diff --git a/Engine/Source/Runtime/Engine/Private/Particles/ParticleSystemRender.cpp b/Engine/Source/Runtime/Engine/Private/Particles/ParticleSystemRender.cpp
--- a/Engine/Source/Runtime/Engine/Private/Particles/ParticleSystemRender.cpp
+++ b/Engine/Source/Runtime/Engine/Private/Particles/ParticleSystemRender.cpp
@@ -961,7 +961,7 @@ void FDynamicSpriteEmitterData::GetDynamicMeshElementsEmitter(const FParticleSys
 	SCOPE_CYCLE_COUNTER(STAT_SpriteRenderingTime);
 
 	const auto FeatureLevel = View->GetFeatureLevel();
-	const bool bInstanced = FeatureLevel >= ERHIFeatureLevel::SM4;
+	const bool bInstanced = RHISupportsInstancing(GetFeatureLevelShaderPlatform(FeatureLevel));
 
 	// Sort and generate particles for this view.
 	const FDynamicSpriteEmitterReplayDataBase* SourceData = GetSourceData();
@@ -1373,7 +1373,7 @@ void FDynamicMeshEmitterData::GetDynamicMeshElementsEmitter(const FParticleSyste
 {
 	SCOPE_CYCLE_COUNTER(STAT_MeshRenderingTime);
 
-	const bool bInstanced = View->GetFeatureLevel() >= ERHIFeatureLevel::SM4;
+	const bool bInstanced = RHISupportsInstancing(GetFeatureLevelShaderPlatform(View->GetFeatureLevel()));
 
 	if (bValid)
 	{
diff --git a/Engine/Source/Runtime/Engine/Private/Particles/ParticleVertexFactory.cpp b/Engine/Source/Runtime/Engine/Private/Particles/ParticleVertexFactory.cpp
--- a/Engine/Source/Runtime/Engine/Private/Particles/ParticleVertexFactory.cpp
+++ b/Engine/Source/Runtime/Engine/Private/Particles/ParticleVertexFactory.cpp
@@ -142,9 +142,9 @@ private:
 static TGlobalResource<FParticleSpriteVertexDeclaration> GParticleSpriteVertexDeclarationInstanced(true);
 static TGlobalResource<FParticleSpriteVertexDeclaration> GParticleSpriteVertexDeclarationNonInstanced(false);
 
-inline TGlobalResource<FParticleSpriteVertexDeclaration>& GetParticleSpriteVertexDeclaration(ERHIFeatureLevel::Type InFeatureLevel)
+inline TGlobalResource<FParticleSpriteVertexDeclaration>& GetParticleSpriteVertexDeclaration(bool SupportsInstancing)
 {
-	if (InFeatureLevel >= ERHIFeatureLevel::SM4)
+	if (SupportsInstancing)
 	{
 		return GParticleSpriteVertexDeclarationInstanced;
 	}
@@ -176,12 +176,12 @@ void FParticleSpriteVertexFactory::ModifyCompilationEnvironment(EShaderPlatform
 void FParticleSpriteVertexFactory::InitRHI()
 {
 	InitStreams();
-	SetDeclaration(GetParticleSpriteVertexDeclaration(GetFeatureLevel()).VertexDeclarationRHI);
+	SetDeclaration(GetParticleSpriteVertexDeclaration(RHISupportsInstancing(GetFeatureLevelShaderPlatform(GetFeatureLevel()))).VertexDeclarationRHI);
 }
 
 void FParticleSpriteVertexFactory::InitStreams()
 {
-	const bool bInstanced = GetFeatureLevel() >= ERHIFeatureLevel::SM4;
+    const bool bInstanced = RHISupportsInstancing(GetFeatureLevelShaderPlatform(GetFeatureLevel()));
 
 	check(Streams.Num() == 0);
 	if(bInstanced) 
diff --git a/Engine/Source/Runtime/IOS/MetalRHI/Private/MetalManager.cpp b/Engine/Source/Runtime/IOS/MetalRHI/Private/MetalManager.cpp
--- a/Engine/Source/Runtime/IOS/MetalRHI/Private/MetalManager.cpp
+++ b/Engine/Source/Runtime/IOS/MetalRHI/Private/MetalManager.cpp
@@ -608,17 +608,7 @@ void FMetalManager::SetRenderTargetsInfo(const FRHISetRenderTargetsInfo& RenderT
 		//			NSLog(@"There were %d draw calls for an RT in frame %lld", NumDrawCalls, GFrameCounter);
 		//		}
 #endif
-		
-		[CurrentContext endEncoding];
-		NumDrawCalls = 0;
-		
-		[CurrentContext release];
-		
-		// commit the buffer for this context
-		[CurrentCommandBuffer commit];
-		UNTRACK_OBJECT(CurrentCommandBuffer);
-		[CurrentCommandBuffer release];
-
+        
 		// Check if CurrentCommandBuffer was rendering to the BackBuffer.
 		if( PreviousRenderTargetsInfo.NumColorRenderTargets == 1 )
 		{
@@ -626,14 +616,11 @@ void FMetalManager::SetRenderTargetsInfo(const FRHISetRenderTargetsInfo& RenderT
 			FMetalSurface& Surface = *GetMetalSurfaceFromRHITexture(RenderTargetView.Texture);
 			if(&Surface == &BackBuffer->Surface && CurrentDrawable != nil)
 			{
-				// release our record of it to ensure we are allocated a new drawable.
+                // release our record of it to ensure we are allocated a new drawable.
 				CurrentDrawable = nil;
 				BackBuffer->Surface.Texture = nil;
 			}
 		}
-		
-		// create the command buffer for this frame
-		CreateCurrentCommandBuffer(false);
 	}
 	
 	// back this up for next frame
@@ -786,7 +773,33 @@ void FMetalManager::SetRenderTargetsInfo(const FRHISetRenderTargetsInfo& RenderT
 	SET_HASH(OFFSET_STENCIL_ENABLED, NUMBITS_STENCIL_ENABLED, (Pipeline.StencilTargetFormat == MTLPixelFormatInvalid ? 0 : 1));
 	SET_HASH(OFFSET_SAMPLE_COUNT, NUMBITS_SAMPLE_COUNT, Pipeline.SampleCount);
 
-	// make a new render context to use to render to the framebuffer
+    // commit pending commands on the old render target
+    if (CurrentContext)
+    {
+#if UE_BUILD_DEBUG || UE_BUILD_DEVELOPMENT
+        //		if (NumDrawCalls == 0)
+        //		{
+        //			NSLog(@"There were %d draw calls for an RT in frame %lld", NumDrawCalls, GFrameCounter);
+        //		}
+#endif
+        
+        [CurrentContext endEncoding];
+        NumDrawCalls = 0;
+        
+        UNTRACK_OBJECT(CurrentContext);
+        [CurrentContext release];
+        CurrentContext = nil;
+        
+        // commit the buffer for this context
+        [CurrentCommandBuffer commit];
+        UNTRACK_OBJECT(CurrentCommandBuffer);
+        [CurrentCommandBuffer release];
+        
+        // create the command buffer for this frame
+        CreateCurrentCommandBuffer(false);
+    }
+
+    // make a new render context to use to render to the framebuffer
 	CurrentContext = [CurrentCommandBuffer renderCommandEncoderWithDescriptor:RenderPass];
 	[CurrentContext retain];
 	TRACK_OBJECT(CurrentContext);
