CIS fix

diff --git a/Engine/Source/Runtime/IOS/MetalRHI/Private/MetalRHI.cpp b/Engine/Source/Runtime/IOS/MetalRHI/Private/MetalRHI.cpp
--- a/Engine/Source/Runtime/IOS/MetalRHI/Private/MetalRHI.cpp
+++ b/Engine/Source/Runtime/IOS/MetalRHI/Private/MetalRHI.cpp
@@ -280,3 +280,18 @@ void* FMetalDynamicRHI::RHIGetNativeDevice()
 {
 	return NULL;
 }
+
+void FMetalDynamicRHI::RHIBeginAsyncComputeJob_DrawThread(EAsyncComputePriority Priority)
+{
+	UE_LOG(LogRHI, Fatal, TEXT("%s not implemented yet"), ANSI_TO_TCHAR(__FUNCTION__));
+}
+
+void FMetalDynamicRHI::RHIEndAsyncComputeJob_DrawThread(uint32 FenceIndex)
+{
+	UE_LOG(LogRHI, Fatal, TEXT("%s not implemented yet"), ANSI_TO_TCHAR(__FUNCTION__));
+}
+
+void FMetalDynamicRHI::RHIGraphicsWaitOnAsyncComputeJob(uint32 FenceIndex)
+{
+	UE_LOG(LogRHI, Fatal, TEXT("%s not implemented yet"), ANSI_TO_TCHAR(__FUNCTION__));
+}
diff --git a/Engine/Source/Runtime/IOS/MetalRHI/Public/MetalRHI.h b/Engine/Source/Runtime/IOS/MetalRHI/Public/MetalRHI.h
--- a/Engine/Source/Runtime/IOS/MetalRHI/Public/MetalRHI.h
+++ b/Engine/Source/Runtime/IOS/MetalRHI/Public/MetalRHI.h
@@ -209,6 +209,9 @@ public:
 	virtual void RHIEnableDepthBoundsTest(bool bEnable, float MinDepth, float MaxDepth) final override;
 	virtual void RHIPushEvent(const TCHAR* Name) final override;
 	virtual void RHIPopEvent() final override;
+	virtual void RHIBeginAsyncComputeJob_DrawThread(EAsyncComputePriority Priority) override;
+	virtual void RHIEndAsyncComputeJob_DrawThread(uint32 FenceIndex) override;
+	virtual void RHIGraphicsWaitOnAsyncComputeJob(uint32 FenceIndex) override;
 
 protected:
     TGlobalResource<TBoundShaderStateHistory<10000>> BoundShaderStateHistory;
