Fix Metal compile errors

diff --git a/Engine/Source/Runtime/IOS/MetalRHI/Private/MetalUAV.cpp b/Engine/Source/Runtime/IOS/MetalRHI/Private/MetalUAV.cpp
--- a/Engine/Source/Runtime/IOS/MetalRHI/Private/MetalUAV.cpp
+++ b/Engine/Source/Runtime/IOS/MetalRHI/Private/MetalUAV.cpp
@@ -105,21 +105,21 @@ FShaderResourceViewRHIRef FMetalDynamicRHI::RHICreateShaderResourceView(FTexture
 
 FShaderResourceViewRHIRef FMetalDynamicRHI::RHICreateShaderResourceView(FTexture3DRHIParamRef Texture3DRHI, uint8 MipLevel)
 {
-	FEmptyShaderResourceView* SRV = new FMetalShaderResourceView;
+	FMetalShaderResourceView* SRV = new FMetalShaderResourceView;
 	SRV->SourceTexture = (FRHITexture*)Texture3DRHI;
 	return SRV;
 }
 
 FShaderResourceViewRHIRef FMetalDynamicRHI::RHICreateShaderResourceView(FTexture2DArrayRHIParamRef Texture2DArrayRHI, uint8 MipLevel)
 {
-	FEmptyShaderResourceView* SRV = new FMetalShaderResourceView;
+	FMetalShaderResourceView* SRV = new FMetalShaderResourceView;
 	SRV->SourceTexture = (FRHITexture*)Texture2DArrayRHI;
 	return SRV;
 }
 
 FShaderResourceViewRHIRef FMetalDynamicRHI::RHICreateShaderResourceView(FTextureCubeRHIParamRef TextureCubeRHI, uint8 MipLevel)
 {
-	FEmptyShaderResourceView* SRV = new FMetalShaderResourceView;
+	FMetalShaderResourceView* SRV = new FMetalShaderResourceView;
 	SRV->SourceTexture = (FRHITexture*)TextureCubeRHI;
 	return SRV;
 }
