UE4 - Compile fixes

diff --git a/Engine/Source/Runtime/IOS/MetalRHI/Public/MetalResources.h b/Engine/Source/Runtime/IOS/MetalRHI/Public/MetalResources.h
--- a/Engine/Source/Runtime/IOS/MetalRHI/Public/MetalResources.h
+++ b/Engine/Source/Runtime/IOS/MetalRHI/Public/MetalResources.h
@@ -73,7 +73,7 @@ public:
 	uint64 DirtyUniformBuffers;
 
 	// List of memory copies from RHIUniformBuffer to packed uniforms
-	TArray<FMetalUniformBufferCopyInfo> UniformBuffersCopyInfo;
+	TArray<CrossCompiler::FUniformBufferCopyInfo> UniformBuffersCopyInfo;
 	
 	TArray<ANSICHAR> GlslCode;
 	NSString* GlslCodeNSString;
@@ -471,7 +471,7 @@ public:
 	 * Commit shader parameters to the currently bound program.
 	 */
 	void CommitPackedGlobals(int32 Stage, const FMetalShaderBindings& Bindings);
-	void CommitPackedUniformBuffers(TRefCountPtr<FMetalBoundShaderState> BoundShaderState, int32 Stage, const TArray< TRefCountPtr<FRHIUniformBuffer> >& UniformBuffers, const TArray<FMetalUniformBufferCopyInfo>& UniformBuffersCopyInfo);
+	void CommitPackedUniformBuffers(TRefCountPtr<FMetalBoundShaderState> BoundShaderState, int32 Stage, const TArray< TRefCountPtr<FRHIUniformBuffer> >& UniformBuffers, const TArray<CrossCompiler::FUniformBufferCopyInfo>& UniformBuffersCopyInfo);
 
 private:
 	/** CPU memory block for storing uniform values. */
