UE-17713 - fix for crash when launching on iOS with metal
#ios

diff --git a/Engine/Source/Runtime/IOS/MetalRHI/Private/MetalManager.cpp b/Engine/Source/Runtime/IOS/MetalRHI/Private/MetalManager.cpp
--- a/Engine/Source/Runtime/IOS/MetalRHI/Private/MetalManager.cpp
+++ b/Engine/Source/Runtime/IOS/MetalRHI/Private/MetalManager.cpp
@@ -183,6 +183,7 @@ FMetalManager::FMetalManager()
 	, SceneFrameCounter(0)
 	, ResourceTableFrameCounter(INDEX_NONE)
 {
+    bPreviousSurfaceWasBackBuffer = false;
 	for (int32 Index = 0; Index < ARRAY_COUNT(CurrentColorRenderTextures); Index++)
 	{
 		CurrentColorRenderTextures[Index] = nil;
@@ -609,12 +610,9 @@ void FMetalManager::SetRenderTargetsInfo(const FRHISetRenderTargetsInfo& RenderT
 		//		}
 #endif
         
-		// Check if CurrentCommandBuffer was rendering to the BackBuffer.
 		if( PreviousRenderTargetsInfo.NumColorRenderTargets == 1 )
 		{
-			const FRHIRenderTargetView& RenderTargetView = PreviousRenderTargetsInfo.ColorRenderTarget[0];
-			FMetalSurface& Surface = *GetMetalSurfaceFromRHITexture(RenderTargetView.Texture);
-			if(&Surface == &BackBuffer->Surface && CurrentDrawable != nil)
+			if(bPreviousSurfaceWasBackBuffer && CurrentDrawable != nil)
 			{
                 // release our record of it to ensure we are allocated a new drawable.
 				CurrentDrawable = nil;
@@ -625,6 +623,13 @@ void FMetalManager::SetRenderTargetsInfo(const FRHISetRenderTargetsInfo& RenderT
 	
 	// back this up for next frame
 	PreviousRenderTargetsInfo = RenderTargetsInfo;
+    // Check if CurrentCommandBuffer was rendering to the BackBuffer.
+    if( RenderTargetsInfo.NumColorRenderTargets == 1 )
+    {
+        const FRHIRenderTargetView& RenderTargetView = RenderTargetsInfo.ColorRenderTarget[0];
+        FMetalSurface* PreviousSurface = GetMetalSurfaceFromRHITexture(RenderTargetView.Texture);
+        bPreviousSurfaceWasBackBuffer = PreviousSurface == &BackBuffer->Surface;
+    }
 
 	FIntPoint MaxDimensions(TNumericLimits<decltype(FIntPoint::X)>::Max(),TNumericLimits<decltype(FIntPoint::Y)>::Max());
 	// at this point, we need to fully set up an encoder/command buffer, so make a new one (autoreleased)
diff --git a/Engine/Source/Runtime/IOS/MetalRHI/Private/MetalManager.h b/Engine/Source/Runtime/IOS/MetalRHI/Private/MetalManager.h
--- a/Engine/Source/Runtime/IOS/MetalRHI/Private/MetalManager.h
+++ b/Engine/Source/Runtime/IOS/MetalRHI/Private/MetalManager.h
@@ -230,7 +230,8 @@ protected:
 	};
 
 	FRHISetRenderTargetsInfo PreviousRenderTargetsInfo;
-
+    bool bPreviousSurfaceWasBackBuffer;
+    
 	FRenderTargetViewInfo CurrentRenderTargetsViewInfo[MaxMetalRenderTargets], PreviousRenderTargetsViewInfo[MaxMetalRenderTargets];
 	uint32 CurrentNumRenderTargets, PreviousNumRenderTargets;
 	id<MTLTexture> CurrentColorRenderTextures[MaxMetalRenderTargets], PreviousColorRenderTextures[MaxMetalRenderTargets];
