PR #1382: Simple spelling fix (Contributed by [REDACTED_PERSON])

diff --git a/Engine/Source/Runtime/IOS/IOSRuntimeSettings/Classes/IOSRuntimeSettings.h b/Engine/Source/Runtime/IOS/IOSRuntimeSettings/Classes/IOSRuntimeSettings.h
--- a/Engine/Source/Runtime/IOS/IOSRuntimeSettings/Classes/IOSRuntimeSettings.h
+++ b/Engine/Source/Runtime/IOS/IOSRuntimeSettings/Classes/IOSRuntimeSettings.h
@@ -216,11 +216,11 @@ public:
 	UPROPERTY(GlobalConfig, EditAnywhere, Category = DeviceOrientations)
 	uint32 bSupportsUpsideDownOrientation : 1;
 
-	// Supports left landscape orientation. Protrait will not be supported.
+	// Supports left landscape orientation. Portrait will not be supported.
 	UPROPERTY(GlobalConfig, EditAnywhere, Category = DeviceOrientations)
 	uint32 bSupportsLandscapeLeftOrientation : 1;
 
-	// Supports right landscape orientation. Protrait will not be supported.
+	// Supports right landscape orientation. Portrait will not be supported.
 	UPROPERTY(GlobalConfig, EditAnywhere, Category = DeviceOrientations)
 	uint32 bSupportsLandscapeRightOrientation : 1;
 
