Ensure stencil buffer is created with metal RHI
Set metal stencil face tests to be the correct way round.
Ensure back face settings are applied when only front face is enabled. (RHI requirement)
Add missing stencil ref value to metal rhi.

Fixes UE-18718

diff --git a/Engine/Source/Runtime/IOS/MetalRHI/Private/MetalCommands.cpp b/Engine/Source/Runtime/IOS/MetalRHI/Private/MetalCommands.cpp
--- a/Engine/Source/Runtime/IOS/MetalRHI/Private/MetalCommands.cpp
+++ b/Engine/Source/Runtime/IOS/MetalRHI/Private/MetalCommands.cpp
@@ -459,6 +459,7 @@ void FMetalDynamicRHI::RHISetDepthStencilState(FDepthStencilStateRHIParamRef New
 	FMetalDepthStencilState* NewState = ResourceCast(NewStateRHI);
 
 	NewState->Set();
+	[FMetalManager::GetContext() setStencilReferenceValue: StencilRef];
 }
 
 void FMetalDynamicRHI::RHISetBlendState(FBlendStateRHIParamRef NewStateRHI, const FLinearColor& BlendFactor)
diff --git a/Engine/Source/Runtime/IOS/MetalRHI/Private/MetalState.cpp b/Engine/Source/Runtime/IOS/MetalRHI/Private/MetalState.cpp
--- a/Engine/Source/Runtime/IOS/MetalRHI/Private/MetalState.cpp
+++ b/Engine/Source/Runtime/IOS/MetalRHI/Private/MetalState.cpp
@@ -199,10 +199,8 @@ FMetalDepthStencilState::FMetalDepthStencilState(const FDepthStencilStateInitial
 		Stencil.depthStencilPassOperation = TranslateStencilOp(Initializer.FrontFacePassStencilOp);
 		Stencil.readMask = Initializer.StencilReadMask;
 		Stencil.writeMask = Initializer.StencilWriteMask;
-		// @todo metal: We had to swap, but I can't remember why these are backwards.
-		// Should the culling stuff be flipped instead? 
-		Desc.backFaceStencil = Stencil;
-		
+
+		Desc.frontFaceStencil = Stencil;
 		[Stencil release];
 	}
 	
@@ -216,11 +214,16 @@ FMetalDepthStencilState::FMetalDepthStencilState(const FDepthStencilStateInitial
 		Stencil.depthStencilPassOperation = TranslateStencilOp(Initializer.BackFacePassStencilOp);
 		Stencil.readMask = Initializer.StencilReadMask;
 		Stencil.writeMask = Initializer.StencilWriteMask;
-		Desc.frontFaceStencil = Stencil;
-		
+
+		Desc.backFaceStencil = Stencil;
 		[Stencil release];
 	}
-			
+	else if(Initializer.bEnableFrontFaceStencil)
+	{
+		// Set backface to match frontface, an RHI requirement.
+		Desc.backFaceStencil = [Desc.frontFaceStencil copy];
+	}
+
 	// bake out the descriptor
 	State = [FMetalManager::GetDevice() newDepthStencilStateWithDescriptor:Desc];
 	TRACK_OBJECT(State);
diff --git a/Engine/Source/Runtime/IOS/MetalRHI/Private/MetalTexture.cpp b/Engine/Source/Runtime/IOS/MetalRHI/Private/MetalTexture.cpp
--- a/Engine/Source/Runtime/IOS/MetalRHI/Private/MetalTexture.cpp
+++ b/Engine/Source/Runtime/IOS/MetalRHI/Private/MetalTexture.cpp
@@ -210,7 +210,7 @@ FMetalSurface::FMetalSurface(ERHIResourceType ResourceType, EPixelFormat Format,
 	}
 
 	// create a stencil buffer if needed
-	if (GMaxRHIShaderPlatform == SP_METAL_MRT && Format == PF_DepthStencil)
+	if (Format == PF_DepthStencil)
 	{
 		Desc.textureType = MTLTextureType2D;
 		Desc.pixelFormat = MTLPixelFormatStencil8;
