UE-18996 popping sound after flying through gates

- Adding clamps to volume output for iOS volume to avoid clipping.

diff --git a/Engine/Source/Runtime/IOS/IOSAudio/Private/IOSAudioSource.cpp b/Engine/Source/Runtime/IOS/IOSAudio/Private/IOSAudioSource.cpp
--- a/Engine/Source/Runtime/IOS/IOSAudio/Private/IOSAudioSource.cpp
+++ b/Engine/Source/Runtime/IOS/IOSAudio/Private/IOSAudioSource.cpp
@@ -257,10 +257,12 @@ void FIOSAudioSoundSource::Update(void)
 	}
 
 	// Apply global multiplier to disable sound when not the foreground app
-	Volume = FMath::Clamp(Volume, 0.0f, MAX_VOLUME) * FApp::GetVolumeMultiplier() * AudioDevice->PlatformAudioHeadroom;
+	Volume *= FApp::GetVolumeMultiplier();
+	Volume *= AudioDevice->PlatformAudioHeadroom;
+	Volume = FMath::Clamp(Volume, 0.0f, 1.0f);
 
 	// Convert to dB
-	const AudioUnitParameterValue Gain = FMath::Clamp<float>(20.0f * log10(Volume), -100, 20.0f);
+	const AudioUnitParameterValue Gain = FMath::Clamp<float>(20.0f * log10(Volume), -100, 0.0f);
 	const AudioUnitParameterValue Pitch = FMath::Clamp<float>(WaveInstance->Pitch, MIN_PITCH, MAX_PITCH);
 
 	OSStatus Status = noErr;
