Update property in ini when override applied to project settings
#ue4
#android
#ios

diff --git a/Engine/Source/Runtime/Android/AndroidRuntimeSettings/Private/AndroidRuntimeSettings.cpp b/Engine/Source/Runtime/Android/AndroidRuntimeSettings/Private/AndroidRuntimeSettings.cpp
--- a/Engine/Source/Runtime/Android/AndroidRuntimeSettings/Private/AndroidRuntimeSettings.cpp
+++ b/Engine/Source/Runtime/Android/AndroidRuntimeSettings/Private/AndroidRuntimeSettings.cpp
@@ -21,12 +21,14 @@ void UAndroidRuntimeSettings::PostEditChangeProperty(struct FPropertyChangedEven
 	if (!bBuildForArmV7 && !bBuildForX86 && !bBuildForX8664)// && !bBuildForArm64)
 	{
 		bBuildForArmV7 = true;
+		UpdateSinglePropertyInConfigFile(GetClass()->FindPropertyByName(GET_MEMBER_NAME_CHECKED(UAndroidRuntimeSettings, bBuildForArmV7)), GetDefaultConfigFilename());
 	}
 
 	// Ensure that at least one GPU architecture is supported
 	if (!bBuildForES2 && !bBuildForES31)
 	{
 		bBuildForES2 = true;
+		UpdateSinglePropertyInConfigFile(GetClass()->FindPropertyByName(GET_MEMBER_NAME_CHECKED(UAndroidRuntimeSettings, bBuildForES2)), GetDefaultConfigFilename());
 	}
 }
 
diff --git a/Engine/Source/Runtime/IOS/IOSRuntimeSettings/Private/IOSRuntimeSettings.cpp b/Engine/Source/Runtime/IOS/IOSRuntimeSettings/Private/IOSRuntimeSettings.cpp
--- a/Engine/Source/Runtime/IOS/IOSRuntimeSettings/Private/IOSRuntimeSettings.cpp
+++ b/Engine/Source/Runtime/IOS/IOSRuntimeSettings/Private/IOSRuntimeSettings.cpp
@@ -39,22 +39,26 @@ void UIOSRuntimeSettings::PostEditChangeProperty(struct FPropertyChangedEvent& P
 	if (!bSupportsPortraitOrientation && !bSupportsUpsideDownOrientation && !bSupportsLandscapeLeftOrientation && !bSupportsLandscapeRightOrientation)
 	{
 		bSupportsPortraitOrientation = true;
+		UpdateSinglePropertyInConfigFile(GetClass()->FindPropertyByName(GET_MEMBER_NAME_CHECKED(UIOSRuntimeSettings, bSupportsPortraitOrientation)), GetDefaultConfigFilename());
 	}
 
 	// Ensure that at least one API is supported
 	if (!bSupportsMetal && !bSupportsOpenGLES2 && !bSupportsMetalMRT)
 	{
 		bSupportsOpenGLES2 = true;
+		UpdateSinglePropertyInConfigFile(GetClass()->FindPropertyByName(GET_MEMBER_NAME_CHECKED(UIOSRuntimeSettings, bSupportsOpenGLES2)), GetDefaultConfigFilename());
 	}
 
 	// Ensure that at least armv7 is selected for shipping and dev
 	if (!bDevForArmV7 && !bDevForArm64 && !bDevForArmV7S)
 	{
 		bDevForArmV7 = true;
+		UpdateSinglePropertyInConfigFile(GetClass()->FindPropertyByName(GET_MEMBER_NAME_CHECKED(UIOSRuntimeSettings, bDevForArmV7)), GetDefaultConfigFilename());
 	}
 	if (!bShipForArmV7 && !bShipForArm64 && !bShipForArmV7S)
 	{
 		bShipForArmV7 = true;
+		UpdateSinglePropertyInConfigFile(GetClass()->FindPropertyByName(GET_MEMBER_NAME_CHECKED(UIOSRuntimeSettings, bShipForArmV7)), GetDefaultConfigFilename());
 	}
 }
 
