UE-20872 - fix for crash on startup when using MSAA on Metal
#ios

diff --git a/Engine/Source/Runtime/IOS/MetalRHI/Private/MetalTexture.cpp b/Engine/Source/Runtime/IOS/MetalRHI/Private/MetalTexture.cpp
--- a/Engine/Source/Runtime/IOS/MetalRHI/Private/MetalTexture.cpp
+++ b/Engine/Source/Runtime/IOS/MetalRHI/Private/MetalTexture.cpp
@@ -212,7 +212,7 @@ FMetalSurface::FMetalSurface(ERHIResourceType ResourceType, EPixelFormat Format,
 	// create a stencil buffer if needed
 	if (Format == PF_DepthStencil)
 	{
-		Desc.textureType = MTLTextureType2D;
+		Desc.textureType = NumSamples > 1 ? MTLTextureType2DMultisample : MTLTextureType2D;
 		Desc.pixelFormat = MTLPixelFormatStencil8;
 		StencilTexture = [FMetalManager::GetDevice() newTextureWithDescriptor:Desc];
 
