Cancel notification by Id if currently displayed on iOS
#rnx

diff --git a/Engine/Source/Runtime/IOS/IOSLocalNotification/Private/IOSLocalNotification.cpp b/Engine/Source/Runtime/IOS/IOSLocalNotification/Private/IOSLocalNotification.cpp
--- a/Engine/Source/Runtime/IOS/IOSLocalNotification/Private/IOSLocalNotification.cpp
+++ b/Engine/Source/Runtime/IOS/IOSLocalNotification/Private/IOSLocalNotification.cpp
@@ -189,7 +189,9 @@ void FIOSLocalNotificationService::CancelLocalNotification(int32 NotificationId)
 {
 #if !PLATFORM_TVOS
 	UNUserNotificationCenter *Center = [UNUserNotificationCenter currentNotificationCenter];
-	[Center removePendingNotificationRequestsWithIdentifiers:@[@(NotificationId).stringValue]];
+	NSArray<NSString*> *Identifiers = @[@(NotificationId).stringValue];
+	[Center removePendingNotificationRequestsWithIdentifiers: Identifiers];
+	[Center removeDeliveredNotificationsWithIdentifiers: Identifiers];
 #endif
 }
 
