Honor the LocalTime bool flag on iOS local notification scheduling.
#rnx

diff --git a/Engine/Source/Runtime/IOS/IOSLocalNotification/Private/IOSLocalNotification.cpp b/Engine/Source/Runtime/IOS/IOSLocalNotification/Private/IOSLocalNotification.cpp
--- a/Engine/Source/Runtime/IOS/IOSLocalNotification/Private/IOSLocalNotification.cpp
+++ b/Engine/Source/Runtime/IOS/IOSLocalNotification/Private/IOSLocalNotification.cpp
@@ -82,7 +82,7 @@ public:
 		
 		return Content;
 	}
-	static UNCalendarNotificationTrigger* CreateCalendarNotificationTrigger(const FDateTime& FireDateTime)
+	static UNCalendarNotificationTrigger* CreateCalendarNotificationTrigger(const FDateTime& FireDateTime, bool LocalTime)
 	{
 		NSCalendar *calendar = [NSCalendar autoupdatingCurrentCalendar];
 		NSDateComponents *dateComps = [[NSDateComponents alloc] init];
@@ -92,6 +92,24 @@ public:
 		[dateComps setHour : FireDateTime.GetHour()];
 		[dateComps setMinute : FireDateTime.GetMinute()];
 		[dateComps setSecond : FireDateTime.GetSecond()];
+		if (!LocalTime)
+		{
+			// if not local time, convert from UTC to local time, UNCalendarNotificationTrigger misbehaves
+			// if the components have anything more than previously set (can't specify timeZone)
+            NSCalendar *utcCalendar = [NSCalendar calendarWithIdentifier:NSCalendarIdentifierISO8601];
+            utcCalendar.timeZone = [NSTimeZone timeZoneForSecondsFromGMT:0];
+            NSDate *utcDate = [utcCalendar dateFromComponents:dateComps];
+            NSDateComponents *utcDateComps = [utcCalendar componentsInTimeZone:calendar.timeZone fromDate:utcDate];
+
+			// now copy components back because utcDateComps has too many fields set for
+			// UNCalendarNotificationTrigger and will now work properly on all devices
+            dateComps.day = utcDateComps.day;
+            dateComps.month = utcDateComps.month;
+            dateComps.year = utcDateComps.year;
+            dateComps.hour = utcDateComps.hour;
+            dateComps.minute = utcDateComps.minute;
+            dateComps.second = utcDateComps.second;
+		}
 		
 		UNCalendarNotificationTrigger *trigger = [UNCalendarNotificationTrigger triggerWithDateMatchingComponents:dateComps repeats:NO];
 		
@@ -135,7 +153,7 @@ int32 FIOSLocalNotificationService::ScheduleLocalNotificationAtTime(const FDateT
 	//have to schedule notification on main thread queue
 	dispatch_async(dispatch_get_main_queue(), ^{
         UNMutableNotificationContent* Content = FIOSLocalNotificationModule::CreateNotificationContent(TitleCopy, BodyCopy, ActionCopy, ActivationEventCopy, 1);
-        UNCalendarNotificationTrigger* Trigger = FIOSLocalNotificationModule::CreateCalendarNotificationTrigger(FireDateTimeCopy);
+        UNCalendarNotificationTrigger* Trigger = FIOSLocalNotificationModule::CreateCalendarNotificationTrigger(FireDateTimeCopy, LocalTime);
         
         UNNotificationRequest *Request = [UNNotificationRequest requestWithIdentifier:@(CurrentNotificationId).stringValue content:Content trigger:Trigger];
 
@@ -162,7 +180,7 @@ int32 FIOSLocalNotificationService::ScheduleLocalNotificationBadgeAtTime(const F
 	//have to schedule notification on main thread queue
 	dispatch_async(dispatch_get_main_queue(), ^{
 		UNMutableNotificationContent* Content = FIOSLocalNotificationModule::CreateNotificationContent(FText(), FText(), FText(), ActivationEventCopy, 1);
-		UNCalendarNotificationTrigger* Trigger = FIOSLocalNotificationModule::CreateCalendarNotificationTrigger(FireDateTime);
+		UNCalendarNotificationTrigger* Trigger = FIOSLocalNotificationModule::CreateCalendarNotificationTrigger(FireDateTime, LocalTime);
 		
 		UNNotificationRequest *Request = [UNNotificationRequest requestWithIdentifier:@(CurrentNotificationId).stringValue content:Content trigger:Trigger];
 
