Add support for new iOS UserNotifications framework, enabled if min iOS version is set >= 10.

diff --git a/Engine/Source/Runtime/ApplicationCore/Private/IOS/IOSAppDelegate.cpp b/Engine/Source/Runtime/ApplicationCore/Private/IOS/IOSAppDelegate.cpp
--- a/Engine/Source/Runtime/ApplicationCore/Private/IOS/IOSAppDelegate.cpp
+++ b/Engine/Source/Runtime/ApplicationCore/Private/IOS/IOSAppDelegate.cpp
@@ -884,10 +884,12 @@ static FAutoConsoleVariableRef CVarGEnableThermalsReport(
     GShowSplashScreen = true;
 	
 #if !PLATFORM_TVOS
+#if __IPHONE_OS_VERSION_MIN_REQUIRED >= __IPHONE_10_0
+	UNUserNotificationCenter *Center = [UNUserNotificationCenter currentNotificationCenter];
+	Center.delegate = self;
+#else
 	// Save launch local notification so the app can check for it when it is ready
-	PRAGMA_DISABLE_DEPRECATION_WARNINGS
 	UILocalNotification *notification = [launchOptions objectForKey:UIApplicationLaunchOptionsLocalNotificationKey];
-	PRAGMA_ENABLE_DEPRECATION_WARNINGS
 	if ( notification != nullptr )
 	{
 		NSDictionary*	userInfo = [notification userInfo];
@@ -903,6 +905,7 @@ static FAutoConsoleVariableRef CVarGEnableThermalsReport(
 			}
 		}
 	}
+#endif
 #endif
 	
     timer = [NSTimer scheduledTimerWithTimeInterval: 0.05f target:self selector:@selector(timerForSplashScreen) userInfo:nil repeats:YES];
@@ -1216,8 +1219,8 @@ FCriticalSection RenderSuspend;
 
 #if !PLATFORM_TVOS && NOTIFICATIONS_ENABLED
 
+#if __IPHONE_OS_VERSION_MIN_REQUIRED < __IPHONE_10_0
 #ifdef __IPHONE_8_0
-PRAGMA_DISABLE_DEPRECATION_WARNINGS
 - (void)application:(UIApplication *)application didRegisterUserNotificationSettings:(UIUserNotificationSettings *)notificationSettings
 {
 	[application registerForRemoteNotifications];
@@ -1227,7 +1230,7 @@ PRAGMA_DISABLE_DEPRECATION_WARNINGS
 		FCoreDelegates::ApplicationRegisteredForUserNotificationsDelegate.Broadcast(types);
     }, TStatId(), NULL, ENamedThreads::GameThread);
 }
-PRAGMA_ENABLE_DEPRECATION_WARNINGS
+#endif
 #endif
 
 - (void)application:(UIApplication *)application didRegisterForRemoteNotificationsWithDeviceToken:(NSData *)deviceToken
@@ -1261,6 +1264,7 @@ PRAGMA_ENABLE_DEPRECATION_WARNINGS
     }, TStatId(), NULL, ENamedThreads::GameThread);
 }
 
+#if __IPHONE_OS_VERSION_MIN_REQUIRED < __IPHONE_10_0
 -(void)application : (UIApplication *)application didReceiveRemoteNotification:(NSDictionary *)userInfo fetchCompletionHandler:(void(^)(UIBackgroundFetchResult result))handler
 {
 	if (bEngineInit && FIOSCoreDelegates::PassesPushNotificationFilters(userInfo))
@@ -1306,11 +1310,115 @@ PRAGMA_ENABLE_DEPRECATION_WARNINGS
 	//   notification.
 	handler(UIBackgroundFetchResultNoData);
 }
+#endif
 
 #endif
 
 #if !PLATFORM_TVOS
-PRAGMA_DISABLE_DEPRECATION_WARNINGS
+
+#if __IPHONE_OS_VERSION_MIN_REQUIRED >= __IPHONE_10_0
+void HandleReceivedNotification(UNNotification* notification)
+{
+	bool IsLocal = false;
+	
+	if ([IOSAppDelegate GetDelegate].bEngineInit)
+	{
+		NSString* NotificationType = (NSString*)[notification.request.content.userInfo objectForKey: @"NotificationType"];
+		if(NotificationType != nullptr)
+		{
+			FString LocalOrRemote(NotificationType);
+			if(LocalOrRemote == FString(TEXT("Local")))
+			{
+				IsLocal = true;
+			}
+		}
+		
+		int AppState;
+		if ([UIApplication sharedApplication].applicationState == UIApplicationStateInactive)
+		{
+			AppState = 1; // EApplicationState::Inactive;
+		}
+		else if ([UIApplication sharedApplication].applicationState == UIApplicationStateBackground)
+		{
+			AppState = 2; // EApplicationState::Background;
+		}
+		else
+		{
+			AppState = 3; // EApplicationState::Active;
+		}
+		
+		if(IsLocal)
+		{
+			NSString*	activationEvent = (NSString*)[notification.request.content.userInfo objectForKey: @"ActivationEvent"];
+			if(activationEvent != nullptr)
+			{
+				FString	activationEventFString(activationEvent);
+				int32	fireDate = [notification.date timeIntervalSince1970];
+				
+				FFunctionGraphTask::CreateAndDispatchWhenReady([activationEventFString, fireDate, AppState]()
+															   {
+																   FCoreDelegates::ApplicationReceivedLocalNotificationDelegate.Broadcast(activationEventFString, fireDate, AppState);
+															   }, TStatId(), NULL, ENamedThreads::GameThread);
+			}
+		}
+		else
+		{
+			NSString* JsonString = @"{}";
+			NSError* JsonError;
+			NSData* JsonData = [NSJSONSerialization dataWithJSONObject : notification.request.content.userInfo
+															   options : 0
+																 error : &JsonError];
+			
+			if (JsonData)
+			{
+				JsonString = [[[NSString alloc] initWithData:JsonData encoding : NSUTF8StringEncoding] autorelease];
+			}
+			
+			FString	jsonFString(JsonString);
+			
+			FFunctionGraphTask::CreateAndDispatchWhenReady([jsonFString, AppState]()
+														   {
+															   FCoreDelegates::ApplicationReceivedRemoteNotificationDelegate.Broadcast(jsonFString, AppState);
+														   }, TStatId(), NULL, ENamedThreads::GameThread);
+		}
+	}
+}
+
+- (void)userNotificationCenter:(UNUserNotificationCenter *)center
+	   willPresentNotification:(UNNotification *)notification
+		 withCompletionHandler:(void (^)(UNNotificationPresentationOptions options))completionHandler
+{
+	// Received notification while app is in the foreground
+	HandleReceivedNotification(notification);
+	
+	completionHandler(UNNotificationPresentationOptionNone);
+}
+
+- (void)userNotificationCenter:(UNUserNotificationCenter *)center
+didReceiveNotificationResponse:(UNNotificationResponse *)response
+		 withCompletionHandler:(void (^)())completionHandler
+{
+	// Received notification while app is in the background or closed
+	
+	// Save launch local notification so the app can check for it when it is ready
+	NSDictionary* userInfo = response.notification.request.content.userInfo;
+	if(userInfo != nullptr)
+	{
+		NSString*	activationEvent = (NSString*)[userInfo objectForKey: @"ActivationEvent"];
+		
+		if(activationEvent != nullptr)
+		{
+			FAppEntry::gAppLaunchedWithLocalNotification = true;
+			FAppEntry::gLaunchLocalNotificationActivationEvent = FString(activationEvent);
+			FAppEntry::gLaunchLocalNotificationFireDate = [response.notification.date timeIntervalSince1970];
+		}
+	}
+	
+	HandleReceivedNotification(response.notification);
+	
+	completionHandler();
+}
+#else
 - (void)application:(UIApplication *)application didReceiveLocalNotification:(UILocalNotification *)notification
 {
 	NSString*	activationEvent = (NSString*)[notification.userInfo objectForKey: @"ActivationEvent"];
@@ -1344,7 +1452,8 @@ PRAGMA_DISABLE_DEPRECATION_WARNINGS
 		NSLog(@"Warning: Missing local notification activation event");
 	}
 }
-PRAGMA_ENABLE_DEPRECATION_WARNINGS
+#endif
+
 #endif
 
 /**
diff --git a/Engine/Source/Runtime/ApplicationCore/Public/IOS/IOSAppDelegate.h b/Engine/Source/Runtime/ApplicationCore/Public/IOS/IOSAppDelegate.h
--- a/Engine/Source/Runtime/ApplicationCore/Public/IOS/IOSAppDelegate.h
+++ b/Engine/Source/Runtime/ApplicationCore/Public/IOS/IOSAppDelegate.h
@@ -9,6 +9,10 @@
 #include "Logging/LogMacros.h"
 #include "Containers/UnrealString.h"
 
+#if __IPHONE_OS_VERSION_MIN_REQUIRED >= __IPHONE_10_0
+#import <UserNotifications/UserNotifications.h>
+#endif
+
 #define USE_MUTE_SWITCH_DETECTION 0
 
 // Predicate to decide whether a push notification message should be processed
@@ -69,6 +73,9 @@ namespace FAppEntry
 	UIGestureRecognizerDelegate,
 #endif
 	GKGameCenterControllerDelegate,
+#if __IPHONE_OS_VERSION_MIN_REQUIRED >= __IPHONE_10_0
+	UNUserNotificationCenterDelegate,
+#endif
 UITextFieldDelegate>
 
 /** Window object */
diff --git a/Engine/Source/Runtime/Core/Core.Build.cs b/Engine/Source/Runtime/Core/Core.Build.cs
--- a/Engine/Source/Runtime/Core/Core.Build.cs
+++ b/Engine/Source/Runtime/Core/Core.Build.cs
@@ -75,7 +75,7 @@ public class Core : ModuleRules
 			AddEngineThirdPartyPrivateStaticDependencies(Target,
 				"zlib"
 				);
-			PublicFrameworks.AddRange(new string[] { "UIKit", "Foundation", "AudioToolbox", "AVFoundation", "GameKit", "StoreKit", "CoreVideo", "CoreMedia", "CoreGraphics", "GameController", "SystemConfiguration", "DeviceCheck" });
+			PublicFrameworks.AddRange(new string[] { "UIKit", "Foundation", "AudioToolbox", "AVFoundation", "GameKit", "StoreKit", "CoreVideo", "CoreMedia", "CoreGraphics", "GameController", "SystemConfiguration", "DeviceCheck", "UserNotifications" });
 			if (Target.Platform == UnrealTargetPlatform.IOS)
 			{
 				PublicFrameworks.AddRange(new string[] { "CoreMotion", "AdSupport", "WebKit" });
diff --git a/Engine/Source/Runtime/Core/Private/IOS/IOSPlatformMisc.cpp b/Engine/Source/Runtime/Core/Private/IOS/IOSPlatformMisc.cpp
--- a/Engine/Source/Runtime/Core/Private/IOS/IOSPlatformMisc.cpp
+++ b/Engine/Source/Runtime/Core/Private/IOS/IOSPlatformMisc.cpp
@@ -47,6 +47,12 @@
 #import <StoreKit/StoreKit.h>
 #import <DeviceCheck/DeviceCheck.h>
 
+#if __IPHONE_OS_VERSION_MIN_REQUIRED >= __IPHONE_10_0
+#import <UserNotifications/UserNotifications.h>
+#include "Async/TaskGraphInterfaces.h"
+#include "Misc/CoreDelegates.h"
+#endif
+
 //#include <libproc.h>
 // @pjs commented out to resolve issue with PLATFORM_TVOS being defined by mach-o loader
 //#include <mach-o/dyld.h>
@@ -965,14 +971,31 @@ void FIOSPlatformMisc::RegisterForRemoteNotifications()
 {
     dispatch_async(dispatch_get_main_queue(), ^{
 #if !PLATFORM_TVOS && NOTIFICATIONS_ENABLED
+#if __IPHONE_OS_VERSION_MIN_REQUIRED >= __IPHONE_10_0
+		UNUserNotificationCenter *Center = [UNUserNotificationCenter currentNotificationCenter];
+		[Center requestAuthorizationWithOptions:(UNAuthorizationOptionBadge | UNAuthorizationOptionSound | UNAuthorizationOptionAlert)
+							  completionHandler:^(BOOL granted, NSError * _Nullable error) {
+								  if (error)
+								  {
+									  UE_LOG(LogIOS, Log, TEXT("Failed to register for notifications."));
+								  }
+								  else
+								  {
+									  int32 types = (int32)granted;
+									  FFunctionGraphTask::CreateAndDispatchWhenReady([types]()
+																					 {
+																						 FCoreDelegates::ApplicationRegisteredForUserNotificationsDelegate.Broadcast(types);
+																					 }, TStatId(), NULL, ENamedThreads::GameThread);
+									  
+								  }
+							  }];
+#else
 	UIApplication* application = [UIApplication sharedApplication];
 	if ([application respondsToSelector : @selector(registerUserNotificationSettings:)])
 	{
 #ifdef __IPHONE_8_0
-		PRAGMA_DISABLE_DEPRECATION_WARNINGS
 		UIUserNotificationSettings * settings = [UIUserNotificationSettings settingsForTypes : (UIUserNotificationTypeBadge | UIUserNotificationTypeSound | UIUserNotificationTypeAlert) categories:nil];
 		[application registerUserNotificationSettings : settings];
-		PRAGMA_ENABLE_DEPRECATION_WARNINGS
 #endif
 	}
 	else
@@ -983,6 +1006,7 @@ void FIOSPlatformMisc::RegisterForRemoteNotifications()
 		[application registerForRemoteNotificationTypes : myTypes];
 #endif
 	}
+#endif
 #endif
     });
 }
@@ -992,6 +1016,25 @@ bool FIOSPlatformMisc::IsRegisteredForRemoteNotifications()
 	return false;
 }
 
+bool FIOSPlatformMisc::IsAllowedRemoteNotifications()
+{
+#if !PLATFORM_TVOS && NOTIFICATIONS_ENABLED
+#if __IPHONE_OS_VERSION_MIN_REQUIRED >= __IPHONE_10_0
+	checkf(false, TEXT("For min iOS version >= 10 use FIOSLocalNotificationService::CheckAllowedNotifications."));
+	return true;
+#elif defined(__IPHONE_8_0)
+	UIApplication* application = [UIApplication sharedApplication];
+	UIUserNotificationSettings* Settings = [application currentUserNotificationSettings];
+	int32 AllowedTypes = (int32)[Settings types];
+	return AllowedTypes != UIUserNotificationTypeNone;
+#else
+	return true;
+#endif
+#else
+	return true;
+#endif
+}
+
 void FIOSPlatformMisc::UnregisterForRemoteNotifications()
 {
 
diff --git a/Engine/Source/Runtime/Core/Public/IOS/IOSPlatformMisc.h b/Engine/Source/Runtime/Core/Public/IOS/IOSPlatformMisc.h
--- a/Engine/Source/Runtime/Core/Public/IOS/IOSPlatformMisc.h
+++ b/Engine/Source/Runtime/Core/Public/IOS/IOSPlatformMisc.h
@@ -57,10 +57,14 @@ struct CORE_API FIOSPlatformMisc : public FApplePlatformMisc
 	static bool SupportsBrightness() { return true; }
     static bool IsInLowPowerMode();
 
+	//////// Notifications
 	static void RegisterForRemoteNotifications();
 	static bool IsRegisteredForRemoteNotifications();
 	static void UnregisterForRemoteNotifications();
-
+	// Check if notifications are allowed if min iOS version is < 10
+	DEPRECATED(4.21, "IsAllowedRemoteNotifications is deprecated. Use FIOSLocalNotificationService::CheckAllowedNotifications instead.")
+	static bool IsAllowedRemoteNotifications();
+	
 	static class IPlatformChunkInstall* GetPlatformChunkInstall();
 
 	static bool SupportsForceTouchInput();
diff --git a/Engine/Source/Runtime/IOS/IOSLocalNotification/Private/IOSLocalNotification.cpp b/Engine/Source/Runtime/IOS/IOSLocalNotification/Private/IOSLocalNotification.cpp
--- a/Engine/Source/Runtime/IOS/IOSLocalNotification/Private/IOSLocalNotification.cpp
+++ b/Engine/Source/Runtime/IOS/IOSLocalNotification/Private/IOSLocalNotification.cpp
@@ -14,6 +14,11 @@
 #include "IOS/IOSAppDelegate.h"
 
 #include "Modules/ModuleManager.h"
+#include "Logging/LogMacros.h"
+
+#if __IPHONE_OS_VERSION_MIN_REQUIRED >= __IPHONE_10_0
+#import <UserNotifications/UserNotifications.h>
+#endif
 
 DEFINE_LOG_CATEGORY(LogIOSLocalNotification);
 
@@ -35,7 +40,67 @@ public:
 	}
 
 #if !PLATFORM_TVOS
-PRAGMA_DISABLE_DEPRECATION_WARNINGS
+#if __IPHONE_OS_VERSION_MIN_REQUIRED >= __IPHONE_10_0
+	static UNMutableNotificationContent* CreateNotificationContent(const FText& Title, const FText& Body, const FText& Action, const FString& ActivationEvent, uint32 BadgeNumber)
+	{
+		UNMutableNotificationContent* Content = [UNMutableNotificationContent new];
+		if(Content != nil)
+		{
+			if(!Title.IsEmpty())
+			{
+				NSString* NotificationTitle = [NSString stringWithFString:Title.ToString()];
+				if(NotificationTitle != nil)
+				{
+					Content.title = NotificationTitle;
+				}
+			}
+			
+			if(!Body.IsEmpty())
+			{
+				NSString* NotificationBody = [NSString stringWithFString:Body.ToString()];
+				if(NotificationBody != nil)
+				{
+					Content.body = NotificationBody;
+				}
+			}
+			
+			NSNumber* BadgeNSNumber = [NSNumber numberWithInt:BadgeNumber];
+			Content.badge = BadgeNSNumber;
+			Content.sound = [UNNotificationSound defaultSound];
+			
+			if(!ActivationEvent.IsEmpty())
+			{
+				NSString* ActivationEventString = [NSString stringWithFString:ActivationEvent];
+				NSString* LocalString = [NSString stringWithFString:FString(TEXT("Local"))];
+				if (ActivationEventString != nil && LocalString != nil)
+				{
+					NSDictionary* Dict = [NSDictionary dictionaryWithObjectsAndKeys: ActivationEventString, @"ActivationEvent", LocalString, @"NotificationType", nil];
+					if (Dict != nil)
+					{
+						Content.userInfo = Dict;
+					}
+				}
+			}
+		}
+		
+		return Content;
+	}
+	static UNCalendarNotificationTrigger* CreateCalendarNotificationTrigger(const FDateTime& FireDateTime)
+	{
+		NSCalendar *calendar = [NSCalendar autoupdatingCurrentCalendar];
+		NSDateComponents *dateComps = [[NSDateComponents alloc] init];
+		[dateComps setDay : FireDateTime.GetDay()];
+		[dateComps setMonth : FireDateTime.GetMonth()];
+		[dateComps setYear : FireDateTime.GetYear()];
+		[dateComps setHour : FireDateTime.GetHour()];
+		[dateComps setMinute : FireDateTime.GetMinute()];
+		[dateComps setSecond : FireDateTime.GetSecond()];
+		
+		UNCalendarNotificationTrigger *trigger = [UNCalendarNotificationTrigger triggerWithDateMatchingComponents:dateComps repeats:NO];
+		
+		return trigger;
+	}
+#else
 	static UILocalNotification* CreateLocalNotification(const FDateTime& FireDateTime, bool bLocalTime, const FString& ActivationEvent)
 	{
 		UIApplication* application = [UIApplication sharedApplication];
@@ -75,7 +140,7 @@ PRAGMA_DISABLE_DEPRECATION_WARNINGS
 		}
 		return localNotif;
 	}
-PRAGMA_ENABLE_DEPRECATION_WARNINGS
+#endif
 #endif
 };
 
@@ -93,18 +158,39 @@ FIOSLocalNotificationService::FIOSLocalNotificationService()
 void FIOSLocalNotificationService::ClearAllLocalNotifications()
 {
 #if !PLATFORM_TVOS
-	PRAGMA_DISABLE_DEPRECATION_WARNINGS
+#if __IPHONE_OS_VERSION_MIN_REQUIRED >= __IPHONE_10_0
+	UNUserNotificationCenter *Center = [UNUserNotificationCenter currentNotificationCenter];
+	
+	[Center removeAllPendingNotificationRequests];
+#else
 	UIApplication* application = [UIApplication sharedApplication];
 	
 	[application cancelAllLocalNotifications];
-	PRAGMA_ENABLE_DEPRECATION_WARNINGS
+#endif
 #endif
 }
 
+static uint32 NotificationNumber = 0;
 void FIOSLocalNotificationService::ScheduleLocalNotificationAtTime(const FDateTime& FireDateTime, bool LocalTime, const FText& Title, const FText& Body, const FText& Action, const FString& ActivationEvent)
 {
 #if !PLATFORM_TVOS
-	PRAGMA_DISABLE_DEPRECATION_WARNINGS
+#if __IPHONE_OS_VERSION_MIN_REQUIRED >= __IPHONE_10_0
+	UNMutableNotificationContent* Content = FIOSLocalNotificationModule::CreateNotificationContent(Title, Body, Action, ActivationEvent, 1);
+	UNCalendarNotificationTrigger* Trigger = FIOSLocalNotificationModule::CreateCalendarNotificationTrigger(FireDateTime);
+	
+	UNUserNotificationCenter *Center = [UNUserNotificationCenter currentNotificationCenter];
+	
+	FString NotId = Title.ToString() + FString::FromInt(NotificationNumber);
+	NSString* NotificationIdentifier = [NSString stringWithFString:NotId];
+	
+	UNNotificationRequest *Request = [UNNotificationRequest requestWithIdentifier:NotificationIdentifier content:Content trigger:Trigger];
+	
+	[Center addNotificationRequest:Request withCompletionHandler:^(NSError * _Nullable error) {
+		if (error != nil) {
+			UE_LOG(LogIOSLocalNotification, Warning, TEXT("Error scheduling notification: %s"), *NotId);
+		}
+	}];
+#else
 	UILocalNotification *localNotif = FIOSLocalNotificationModule::CreateLocalNotification(FireDateTime, LocalTime, ActivationEvent);
 	if (localNotif == nil)
 		return;
@@ -134,14 +220,30 @@ void FIOSLocalNotificationService::ScheduleLocalNotificationAtTime(const FDateTi
 	localNotif.applicationIconBadgeNumber = 1;
 
 	[[UIApplication sharedApplication] scheduleLocalNotification:localNotif];
-	PRAGMA_ENABLE_DEPRECATION_WARNINGS
+#endif
 #endif
 }
 
 void FIOSLocalNotificationService::ScheduleLocalNotificationBadgeAtTime(const FDateTime& FireDateTime, bool LocalTime, const FString& ActivationEvent)
 {
 #if !PLATFORM_TVOS
-	PRAGMA_DISABLE_DEPRECATION_WARNINGS
+#if __IPHONE_OS_VERSION_MIN_REQUIRED >= __IPHONE_10_0
+	UNMutableNotificationContent* Content = FIOSLocalNotificationModule::CreateNotificationContent(FText(), FText(), FText(), ActivationEvent, 1);
+	UNCalendarNotificationTrigger* Trigger = FIOSLocalNotificationModule::CreateCalendarNotificationTrigger(FireDateTime);
+	
+	UNUserNotificationCenter *Center = [UNUserNotificationCenter currentNotificationCenter];
+	
+	FString NotId = FString(TEXT("Badge")) + FString::FromInt(NotificationNumber);
+	NSString* NotificationIdentifier = [NSString stringWithFString:NotId];
+	
+	UNNotificationRequest *Request = [UNNotificationRequest requestWithIdentifier:NotificationIdentifier content:Content trigger:Trigger];
+	
+	[Center addNotificationRequest:Request withCompletionHandler:^(NSError * _Nullable error) {
+		if (error != nil) {
+			UE_LOG(LogIOSLocalNotification, Warning, TEXT("Error scheduling notification: %s"), *NotId);
+		}
+	}];
+#else
 	UILocalNotification *localNotif = FIOSLocalNotificationModule::CreateLocalNotification(FireDateTime, LocalTime, ActivationEvent);
 	if (localNotif == nil)
 		return;
@@ -152,7 +254,7 @@ void FIOSLocalNotificationService::ScheduleLocalNotificationBadgeAtTime(const FD
 	localNotif.applicationIconBadgeNumber = 1;
 
 	[[UIApplication sharedApplication] scheduleLocalNotification:localNotif];
-	PRAGMA_ENABLE_DEPRECATION_WARNINGS
+#endif
 #endif
 }
 
@@ -174,3 +276,21 @@ void FIOSLocalNotificationService::SetLaunchNotification(FString const& Activati
 	LaunchNotificationActivationEvent = ActivationEvent;
 	LaunchNotificationFireDate = FireDate;
 }
+
+
+static FIOSLocalNotificationService::FAllowedNotifications NotificationsAllowedDelegate;
+void FIOSLocalNotificationService::CheckAllowedNotifications(const FAllowedNotifications& AllowedNotificationsDelegate)
+{
+#if __IPHONE_OS_VERSION_MIN_REQUIRED >= __IPHONE_10_0
+	NotificationsAllowedDelegate = AllowedNotificationsDelegate;
+	
+	UNUserNotificationCenter *Center = [UNUserNotificationCenter currentNotificationCenter];
+	[Center getNotificationSettingsWithCompletionHandler:^(UNNotificationSettings * _Nonnull settings) {
+		bool NotificationsAllowed = settings.authorizationStatus == UNAuthorizationStatusAuthorized;
+		NotificationsAllowedDelegate.ExecuteIfBound(NotificationsAllowed);
+	}];
+	
+#else
+	checkf(false, TEXT("For min iOS version < 10 use FIOSPlatformMisc::IsAllowedRemoteNotifications()."));
+#endif
+}
diff --git a/Engine/Source/Runtime/IOS/IOSLocalNotification/Public/IOSLocalNotification.h b/Engine/Source/Runtime/IOS/IOSLocalNotification/Public/IOSLocalNotification.h
--- a/Engine/Source/Runtime/IOS/IOSLocalNotification/Public/IOSLocalNotification.h
+++ b/Engine/Source/Runtime/IOS/IOSLocalNotification/Public/IOSLocalNotification.h
@@ -8,6 +8,7 @@
 
 #include "LocalNotification.h"
 #include "Logging/LogMacros.h"
+#include "Delegates/Delegate.h"
 
 DECLARE_LOG_CATEGORY_EXTERN(LogIOSLocalNotification, Log, All);
 
@@ -58,6 +59,10 @@ public:
 	*/
 	virtual void CancelLocalNotification(const FString& ActivationEvent);
 
+	// Check if notifications are allowed if min iOS version is > 10
+	DECLARE_DELEGATE_OneParam(FAllowedNotifications, bool);
+	static void CheckAllowedNotifications(const FAllowedNotifications& AllowedNotificationsDelegate);
+
 private:
 	bool	AppLaunchedWithNotification;
 	FString	LaunchNotificationActivationEvent;
