[NO NON-PII COMMIT MESSAGE]

diff --git a/Engine/Source/Runtime/IOS/IOSAudio/Private/IOSAudioBuffer.cpp b/Engine/Source/Runtime/IOS/IOSAudio/Private/IOSAudioBuffer.cpp
--- a/Engine/Source/Runtime/IOS/IOSAudio/Private/IOSAudioBuffer.cpp
+++ b/Engine/Source/Runtime/IOS/IOSAudio/Private/IOSAudioBuffer.cpp
@@ -21,20 +21,31 @@
 	FIOSAudioSoundBuffer
  ------------------------------------------------------------------------------------*/
 
-FIOSAudioSoundBuffer::FIOSAudioSoundBuffer(FIOSAudioDevice* InAudioDevice, USoundWave* InWave, bool InStreaming):
+FIOSAudioSoundBuffer::FIOSAudioSoundBuffer(FIOSAudioDevice* InAudioDevice, USoundWave* InWave, bool InStreaming, bool InProcedural):
 	FSoundBuffer(InAudioDevice),
-	SampleData(NULL),
+	SampleData(nullptr),
 	BufferSize(0),
-	bStreaming(InStreaming)
+	DecompressionState(nullptr),
+	bStreaming(InStreaming),
+	bIsProcedural(InProcedural)
 {
-	DecompressionState = static_cast<FADPCMAudioInfo*>(InAudioDevice->CreateCompressedAudioInfo(InWave));
-
-	if (!ReadCompressedInfo(InWave))
+	if (!bIsProcedural)
 	{
-		return;
+		DecompressionState = static_cast<FADPCMAudioInfo*>(InAudioDevice->CreateCompressedAudioInfo(InWave));
+
+		if (!ReadCompressedInfo(InWave))
+		{
+			return;
+		}
+
+		SoundFormat = static_cast<ESoundFormat>(*DecompressionState->WaveInfo.pFormatTag);
 	}
+    else
+    {
+        SoundFormat = SoundFormat_LPCM;
+    }
+	
 	
-	SoundFormat = static_cast<ESoundFormat>(*DecompressionState->WaveInfo.pFormatTag);
 	SampleRate = InWave->GetSampleRateForCurrentPlatform();
 	NumChannels = InWave->NumChannels;
 	BufferSize = AudioCallbackFrameSize * sizeof(uint16) * NumChannels;
@@ -57,16 +68,16 @@ FIOSAudioSoundBuffer::~FIOSAudioSoundBuffer(void)
 		UE_LOG(LogIOSAudio, Fatal, TEXT("Can't free resource '%s' as it was allocated in permanent pool."), *ResourceName);
 	}
 	
-	if(SampleData != NULL)
+	if(SampleData != nullptr)
 	{
 		FMemory::Free(SampleData);
-		SampleData = NULL;
+		SampleData = nullptr;
 	}
 	
-	if(DecompressionState != NULL)
+	if(DecompressionState != nullptr)
 	{
 		delete DecompressionState;
-		DecompressionState = NULL;
+		DecompressionState = nullptr;
 	}
 }
 
@@ -87,7 +98,7 @@ int32 FIOSAudioSoundBuffer::GetSize(void)
 
 int32 FIOSAudioSoundBuffer::GetCurrentChunkIndex() const
 {
-	if(DecompressionState == NULL)
+	if(DecompressionState == nullptr)
 	{
 		return -1;
 	}
@@ -97,7 +108,7 @@ int32 FIOSAudioSoundBuffer::GetCurrentChunkIndex() const
 
 int32 FIOSAudioSoundBuffer::GetCurrentChunkOffset() const
 {
-	if(DecompressionState == NULL)
+	if(DecompressionState == nullptr)
 	{
 		return -1;
 	}
@@ -107,6 +118,8 @@ int32 FIOSAudioSoundBuffer::GetCurrentChunkOffset() const
 
 bool FIOSAudioSoundBuffer::ReadCompressedInfo(USoundWave* InWave)
 {
+	check(DecompressionState != nullptr)
+
 	FSoundQualityInfo QualityInfo = { 0 };
 	if(bStreaming)
 	{
@@ -129,7 +142,10 @@ bool FIOSAudioSoundBuffer::ReadCompressedData( uint8* Destination, int32 NumFram
 	{
 		return DecompressionState->StreamCompressedData(Destination, bLooping, RenderCallbackBufferSize);
 	}
-	return DecompressionState->ReadCompressedData(Destination, bLooping, RenderCallbackBufferSize);
+	else
+	{
+		return DecompressionState->ReadCompressedData(Destination, bLooping, RenderCallbackBufferSize);
+	}
 }
 
 FIOSAudioSoundBuffer* FIOSAudioSoundBuffer::Init(FIOSAudioDevice* IOSAudioDevice, USoundWave* InWave)
@@ -160,13 +176,16 @@ FIOSAudioSoundBuffer* FIOSAudioSoundBuffer::Init(FIOSAudioDevice* IOSAudioDevice
 						
 		case DTYPE_RealTime:
 			// Always create a new FIOSAudioSoundBuffer since positional information about the sound is kept track of in this object
-			Buffer = new FIOSAudioSoundBuffer(IOSAudioDevice, InWave, bIsStreaming);
+			Buffer = new FIOSAudioSoundBuffer(IOSAudioDevice, InWave, bIsStreaming, false);
 			break;
-			
+
+		case DTYPE_Procedural:
+			Buffer = new FIOSAudioSoundBuffer(IOSAudioDevice, InWave, bIsStreaming, true);
+			break;
+
 		case DTYPE_Native:
 		case DTYPE_Invalid:
 		case DTYPE_Preview:
-		case DTYPE_Procedural:
 		default:
 			// Invalid will be set if the wave cannot be played
 			UE_LOG( LogIOSAudio, Warning, TEXT("Init Buffer on unsupported sound type name = %s type = %d"), *InWave->GetName(), int32(InWave->DecompressionType));
diff --git a/Engine/Source/Runtime/IOS/IOSAudio/Private/IOSAudioSource.cpp b/Engine/Source/Runtime/IOS/IOSAudio/Private/IOSAudioSource.cpp
--- a/Engine/Source/Runtime/IOS/IOSAudio/Private/IOSAudioSource.cpp
+++ b/Engine/Source/Runtime/IOS/IOSAudio/Private/IOSAudioSource.cpp
@@ -464,21 +464,38 @@ OSStatus FIOSAudioSoundSource::IOSAudioRenderCallback(void* RefCon, AudioUnitRen
 		FMemory::Memzero(OutData, NumFrames * sizeof(AudioSampleType));
 		return -1;
 	}
-	
+
+    
+    
 	if(Channel == 0)
 	{
-		Source->IOSBuffer->RenderCallbackBufferSize = NumFrames * sizeof(uint16) * Source->IOSBuffer->DecompressionState->NumChannels;
-		
-		// Since StreamCompressedData returns interlaced samples we need to decompress all frames(samples) for all channels here so we don't end up decompressing multiple times
-		// Ensure we have enough memory to do this. If needed we could realloc here but that is bad practice inside the audio callback since it has a hard deadline
-		check(Source->IOSBuffer->RenderCallbackBufferSize <= Source->IOSBuffer->BufferSize)
-		
-		// Grab the interlaced channel data
-		Source->bChannel0Finished =
-			Source->IOSBuffer->ReadCompressedData(
+		// Grab the interleaved channel data
+		if (Source->IOSBuffer->bIsProcedural)
+		{
+            Source->IOSBuffer->RenderCallbackBufferSize = NumFrames * sizeof(uint16) * Source->IOSBuffer->NumChannels;
+            check(Source->IOSBuffer->RenderCallbackBufferSize <= Source->IOSBuffer->BufferSize);
+            
+            FMemory::Memzero(Source->IOSBuffer->SampleData, Source->IOSBuffer->RenderCallbackBufferSize);
+            
+			int32 DataSize = Source->WaveInstance->WaveData->GeneratePCMData(
+				(uint8*) Source->IOSBuffer->SampleData,
+				Source->IOSBuffer->RenderCallbackBufferSize / sizeof(int16));
+			Source->bChannel0Finished = DataSize <= 0;
+		}
+		else
+		{
+			Source->IOSBuffer->RenderCallbackBufferSize = NumFrames * sizeof(uint16) * Source->IOSBuffer->DecompressionState->NumChannels;
+
+			// Since StreamCompressedData returns interlaced samples we need to decompress all frames(samples) for all channels here so we don't end up decompressing multiple times
+			// Ensure we have enough memory to do this. If needed we could realloc here but that is bad practice inside the audio callback since it has a hard deadline
+			check(Source->IOSBuffer->RenderCallbackBufferSize <= Source->IOSBuffer->BufferSize);
+
+			Source->bChannel0Finished =
+				Source->IOSBuffer->ReadCompressedData(
 				(uint8*)Source->IOSBuffer->SampleData,
 				MONO_PCM_BUFFER_SAMPLES,
 				Source->WaveInstance->LoopingMode == LOOP_WithNotification || Source->WaveInstance->LoopingMode == LOOP_Forever);
+		}
 	}
     
     // If the channel count is higher than we've expected,
diff --git a/Engine/Source/Runtime/IOS/IOSAudio/Public/IOSAudioDevice.h b/Engine/Source/Runtime/IOS/IOSAudio/Public/IOSAudioDevice.h
--- a/Engine/Source/Runtime/IOS/IOSAudio/Public/IOSAudioDevice.h
+++ b/Engine/Source/Runtime/IOS/IOSAudio/Public/IOSAudioDevice.h
@@ -48,7 +48,7 @@ enum ESoundFormat
 class FIOSAudioSoundBuffer : public FSoundBuffer
 {
 public:
-	FIOSAudioSoundBuffer(FIOSAudioDevice* InAudioDevice, USoundWave* InWave, bool InStreaming);
+	FIOSAudioSoundBuffer(FIOSAudioDevice* InAudioDevice, USoundWave* InWave, bool InStreaming, bool InProcedural);
 	virtual ~FIOSAudioSoundBuffer(void);
 
 	/**
@@ -95,6 +95,7 @@ public:
 	/** Wrapper to assist in the bookkeeping of uncompressed data when streaming */
 	class FADPCMAudioInfo*		DecompressionState;
 	bool	bStreaming;
+	bool    bIsProcedural;
 };
 
 /**
@@ -200,6 +201,7 @@ public:
 	AUGraph GetAudioUnitGraph() const { return AudioUnitGraph; }
 	AUNode GetMixerNode() const { return MixerNode; }
 	AudioUnit GetMixerUnit() const { return MixerUnit; }
+	AudioUnit GetOutputUnit() const { return OutputUnit; }
 
 	/** Thread context management */
 	virtual void ResumeContext();
