[NO NON-PII COMMIT MESSAGE]

diff --git a/Engine/Source/Runtime/IOS/IOSAudio/Private/IOSAudioSource.cpp b/Engine/Source/Runtime/IOS/IOSAudio/Private/IOSAudioSource.cpp
--- a/Engine/Source/Runtime/IOS/IOSAudio/Private/IOSAudioSource.cpp
+++ b/Engine/Source/Runtime/IOS/IOSAudio/Private/IOSAudioSource.cpp
@@ -226,7 +226,7 @@ void FIOSAudioSoundSource::Update(void)
 
     SetFilterFrequency();
     
-    SourceLPFFrequency = LPFFrequency;
+    SourceLPFFrequency = LPFFrequency / (((float) SampleRate) * 0.5f);
     
 	// Factor in the xaudio2 attenuation that happens to stereo assets.
 	if (WaveInstance->WaveData->NumChannels == 2 && WaveInstance->bUseSpatialization)
@@ -292,10 +292,9 @@ void FIOSAudioSoundSource::Play(void)
 		// Updates the source which sets the pitch and volume
 		Update();
         
-        for (Audio::FBiquadFilter& Filter : LowpassFilterBank)
+        for (Audio::FOnePoleLPF& Filter : LowpassFilterBank)
         {
             Filter.Reset();
-            Filter.Init(SampleRate, 1, Audio::EBiquadFilter::Lowpass, LPFFrequency);
         }
         
         for (Audio::FParam& Param : LPFParamBank)
@@ -509,12 +508,11 @@ OSStatus FIOSAudioSoundSource::IOSAudioRenderCallback(void* RefCon, AudioUnitRen
         
         ChannelIndex = Source->LowpassFilterBank.AddDefaulted(1);
         Source->LowpassFilterBank[ChannelIndex].Reset();
-        Source->LowpassFilterBank[ChannelIndex].Init(Source->SampleRate, 1, Audio::EBiquadFilter::Lowpass, Source->SourceLPFFrequency);
     }
     
     // Set up LPF filter for this channel:
     Audio::FParam& LPFParam = Source->LPFParamBank[Channel];
-    Audio::FBiquadFilter& LowpassFilter = Source->LowpassFilterBank[Channel];
+    Audio::FOnePoleLPF& LowpassFilter = Source->LowpassFilterBank[Channel];
     
     LPFParam.Reset();
     const bool bShouldUpdateCutoffFrequency = !FMath::IsNearlyEqual(Source->SourceLPFFrequency, LPFParam.GetValue());
@@ -541,7 +539,7 @@ OSStatus FIOSAudioSoundSource::IOSAudioRenderCallback(void* RefCon, AudioUnitRen
         float FloatSample = ((float) IntSample) / 32767.0f;
         if(LPFParam.GetValue() != MAX_FILTER_FREQUENCY)
         {
-            LowpassFilter.ProcessAudioFrame(&FloatSample, &FloatSample);
+            FloatSample = LowpassFilter.ProcessAudioSample(FloatSample);
         }
         *OutData++ = (SInt16) (FloatSample * 32767.0f);
 	}
diff --git a/Engine/Source/Runtime/IOS/IOSAudio/Public/IOSAudioDevice.h b/Engine/Source/Runtime/IOS/IOSAudio/Public/IOSAudioDevice.h
--- a/Engine/Source/Runtime/IOS/IOSAudio/Public/IOSAudioDevice.h
+++ b/Engine/Source/Runtime/IOS/IOSAudio/Public/IOSAudioDevice.h
@@ -11,7 +11,7 @@
 #include "AudioDevice.h"
 #include "Sound/SoundWave.h"
 #include "DSP/ParamInterpolator.h"
-#include "DSP/Filter.h"
+#include "DSP/OnePole.h"
 
 /*------------------------------------------------------------------------------------
 	Audio Framework system headers
@@ -150,7 +150,7 @@ public:
 	/** Calculates the audio unit element of the input channel relative to the base bus number */
 	AudioUnitElement GetAudioUnitElement(int32 Channel);
 
-    TArray<Audio::FBiquadFilter> LowpassFilterBank;
+    TArray<Audio::FOnePoleLPF> LowpassFilterBank;
     TArray<Audio::FParam> LPFParamBank;
     
     int32 SampleRate;
