UE-65327 Project with Preferred Landscape Orientation set to Left Home Button still launches oriented Right Home Button

#4.22
#iOS

Removing unnecessary code (forcing flags from Engine/Config/BaseEngine.ini to Config/DefaultEngine.ini)

diff --git a/Engine/Source/Runtime/IOS/IOSRuntimeSettings/Classes/IOSRuntimeSettings.h b/Engine/Source/Runtime/IOS/IOSRuntimeSettings/Classes/IOSRuntimeSettings.h
--- a/Engine/Source/Runtime/IOS/IOSRuntimeSettings/Classes/IOSRuntimeSettings.h
+++ b/Engine/Source/Runtime/IOS/IOSRuntimeSettings/Classes/IOSRuntimeSettings.h
@@ -513,7 +513,4 @@ public:
 	// End of UObject interface
 #endif
 
-private:
-	virtual void EnsureOrientationInProjectDefaultEngine();
-
 };
diff --git a/Engine/Source/Runtime/IOS/IOSRuntimeSettings/Private/IOSRuntimeSettings.cpp b/Engine/Source/Runtime/IOS/IOSRuntimeSettings/Private/IOSRuntimeSettings.cpp
--- a/Engine/Source/Runtime/IOS/IOSRuntimeSettings/Private/IOSRuntimeSettings.cpp
+++ b/Engine/Source/Runtime/IOS/IOSRuntimeSettings/Private/IOSRuntimeSettings.cpp
@@ -185,15 +185,5 @@ void UIOSRuntimeSettings::PostInitProperties()
 		MinimumiOSVersion = EIOSVersion::IOS_10;
 	}
 
-	EnsureOrientationInProjectDefaultEngine();
 }
 #endif
-
-void UIOSRuntimeSettings::EnsureOrientationInProjectDefaultEngine()
-{
-	// Make sure the values from BaseEngine.ini are written in the project's DefaultEngine.ini
-	UpdateSinglePropertyInConfigFile(GetClass()->FindPropertyByName(GET_MEMBER_NAME_CHECKED(UIOSRuntimeSettings, bSupportsPortraitOrientation)), GetDefaultConfigFilename());
-	UpdateSinglePropertyInConfigFile(GetClass()->FindPropertyByName(GET_MEMBER_NAME_CHECKED(UIOSRuntimeSettings, bSupportsUpsideDownOrientation)), GetDefaultConfigFilename());
-	UpdateSinglePropertyInConfigFile(GetClass()->FindPropertyByName(GET_MEMBER_NAME_CHECKED(UIOSRuntimeSettings, bSupportsLandscapeLeftOrientation)), GetDefaultConfigFilename());
-	UpdateSinglePropertyInConfigFile(GetClass()->FindPropertyByName(GET_MEMBER_NAME_CHECKED(UIOSRuntimeSettings, PreferredLandscapeOrientation)), GetDefaultConfigFilename());
-}
