Change the name of the "Metal Desktop Forward Renderer" to simply "Metal Desktop Renderer" on iOS as it can run the Deferred or Forward renderer, albeit the Deferred renderer is typically much slower.

diff --git a/Engine/Source/Runtime/IOS/IOSRuntimeSettings/Classes/IOSRuntimeSettings.h b/Engine/Source/Runtime/IOS/IOSRuntimeSettings/Classes/IOSRuntimeSettings.h
--- a/Engine/Source/Runtime/IOS/IOSRuntimeSettings/Classes/IOSRuntimeSettings.h
+++ b/Engine/Source/Runtime/IOS/IOSRuntimeSettings/Classes/IOSRuntimeSettings.h
@@ -228,8 +228,8 @@ public:
 	UPROPERTY(GlobalConfig, EditAnywhere, Category = Rendering, meta = (DisplayName = "Metal Mobile Renderer"))
 	bool bSupportsMetal;
 
-	// Whether or not to compile iOS Metal shaders for the desktop Forward renderer (requires iOS 10+ and an A10 processor)
-	UPROPERTY(GlobalConfig, EditAnywhere, Category = Rendering, meta = (DisplayName = "Metal Desktop-Forward Renderer"))
+	// Whether or not to compile iOS Metal shaders for the desktop renderer (requires iOS 10+ and an A10 processor)
+	UPROPERTY(GlobalConfig, EditAnywhere, Category = Rendering, meta = (DisplayName = "Metal Desktop Renderer"))
 	bool bSupportsMetalMRT;
 	
 	// Whether or not to add support for PVRTC textures
