turn on bitcode is shipping by default

diff --git a/Engine/Config/BaseEngine.ini b/Engine/Config/BaseEngine.ini
--- a/Engine/Config/BaseEngine.ini
+++ b/Engine/Config/BaseEngine.ini
@@ -1840,6 +1840,7 @@ bDevForArmV7S=False
 bShipForArmV7=False
 bShipForArm64=True
 bShipForArmV7S=False
+bShipForBitcode=True
 bTreatRemoteAsSeparateController=False
 bAllowRemoteRotation=True
 bUseRemoteAsVirtualJoystick=True
diff --git a/Engine/Source/Runtime/IOS/IOSRuntimeSettings/Private/IOSRuntimeSettings.cpp b/Engine/Source/Runtime/IOS/IOSRuntimeSettings/Private/IOSRuntimeSettings.cpp
--- a/Engine/Source/Runtime/IOS/IOSRuntimeSettings/Private/IOSRuntimeSettings.cpp
+++ b/Engine/Source/Runtime/IOS/IOSRuntimeSettings/Private/IOSRuntimeSettings.cpp
@@ -32,7 +32,7 @@ UIOSRuntimeSettings::UIOSRuntimeSettings(const FObjectInitializer& ObjectInitial
 	bShipForArmV7 = false;
 	bShipForArm64 = true;
 	bShipForArmV7S = false;
-	bShipForBitcode = false;
+	bShipForBitcode = true;
 	bUseRSync = true;
 	AdditionalPlistData = TEXT("");
 	AdditionalLinkerFlags = TEXT("");
