- Fixed local notification service on IOS to have unique identifiers

diff --git a/Engine/Source/Runtime/IOS/IOSLocalNotification/Private/IOSLocalNotification.cpp b/Engine/Source/Runtime/IOS/IOSLocalNotification/Private/IOSLocalNotification.cpp
--- a/Engine/Source/Runtime/IOS/IOSLocalNotification/Private/IOSLocalNotification.cpp
+++ b/Engine/Source/Runtime/IOS/IOSLocalNotification/Private/IOSLocalNotification.cpp
@@ -190,11 +190,12 @@ void FIOSLocalNotificationService::ScheduleLocalNotificationAtTime(const FDateTi
         
         UNUserNotificationCenter *Center = [UNUserNotificationCenter currentNotificationCenter];
         
-        FString NotId = TitleCopy.ToString() + FString::FromInt(NotificationNumber);
+        FString NotId = TitleCopy.ToString() + FString::FromInt(NotificationNumber++);
         NSString* NotificationIdentifier = [NSString stringWithFString : NotId];
-        
+		
         UNNotificationRequest *Request = [UNNotificationRequest requestWithIdentifier:NotificationIdentifier content:Content trigger:Trigger];
-        
+		NSLog(@"NOTIFICATION: %@ [%@]", NotificationIdentifier, Request);
+
         [Center addNotificationRequest : Request withCompletionHandler : ^ (NSError * _Nullable error) {
             if (error != nil) {
                 UE_LOG(LogIOSLocalNotification, Warning, TEXT("Error scheduling notification: %s"), *NotId);
@@ -255,7 +256,7 @@ void FIOSLocalNotificationService::ScheduleLocalNotificationBadgeAtTime(const FD
 	
 	UNUserNotificationCenter *Center = [UNUserNotificationCenter currentNotificationCenter];
 	
-	FString NotId = FString(TEXT("Badge")) + FString::FromInt(NotificationNumber);
+	FString NotId = FString(TEXT("Badge")) + FString::FromInt(NotificationNumber++);
 	NSString* NotificationIdentifier = [NSString stringWithFString:NotId];
 	
 	UNNotificationRequest *Request = [UNNotificationRequest requestWithIdentifier:NotificationIdentifier content:Content trigger:Trigger];
