- Removed debug log

diff --git a/Engine/Source/Runtime/IOS/IOSLocalNotification/Private/IOSLocalNotification.cpp b/Engine/Source/Runtime/IOS/IOSLocalNotification/Private/IOSLocalNotification.cpp
--- a/Engine/Source/Runtime/IOS/IOSLocalNotification/Private/IOSLocalNotification.cpp
+++ b/Engine/Source/Runtime/IOS/IOSLocalNotification/Private/IOSLocalNotification.cpp
@@ -194,7 +194,6 @@ void FIOSLocalNotificationService::ScheduleLocalNotificationAtTime(const FDateTi
         NSString* NotificationIdentifier = [NSString stringWithFString : NotId];
 		
         UNNotificationRequest *Request = [UNNotificationRequest requestWithIdentifier:NotificationIdentifier content:Content trigger:Trigger];
-		NSLog(@"NOTIFICATION: %@ [%@]", NotificationIdentifier, Request);
 
         [Center addNotificationRequest : Request withCompletionHandler : ^ (NSError * _Nullable error) {
             if (error != nil) {
