- Cleaned up the LocalNotification service code on IOS and Android, and allow for canceling a notification by ID
- Android ID changes untested yet

diff --git a/Engine/Build/Android/Java/src/com/epicgames/ue4/GameActivity.java.template b/Engine/Build/Android/Java/src/com/epicgames/ue4/GameActivity.java.template
--- a/Engine/Build/Android/Java/src/com/epicgames/ue4/GameActivity.java.template
+++ b/Engine/Build/Android/Java/src/com/epicgames/ue4/GameActivity.java.template
@@ -4769,7 +4769,6 @@ public class GameActivity extends $${gameActivitySuperClass}$$ implements Surfac
 		}
 	}
 
-	/*
 	// Returns true only if the scheduled notification exists and gets destoyed successfully
 	public boolean AndroidThunkJava_LocalNotificationDestroyIfExists(int notificationId)
 	{
@@ -4784,7 +4783,6 @@ public class GameActivity extends $${gameActivitySuperClass}$$ implements Surfac
 		}
 		return false;
 	}
-	*/
 
 	// List of vendor/product ids
 	private static final DeviceInfoData[] DeviceInfoList = {
diff --git a/Engine/Source/Runtime/Android/AndroidLocalNotification/Private/AndroidLocalNotification.cpp b/Engine/Source/Runtime/Android/AndroidLocalNotification/Private/AndroidLocalNotification.cpp
--- a/Engine/Source/Runtime/Android/AndroidLocalNotification/Private/AndroidLocalNotification.cpp
+++ b/Engine/Source/Runtime/Android/AndroidLocalNotification/Private/AndroidLocalNotification.cpp
@@ -32,7 +32,7 @@ public:
 };
 
 
-#if PLATFORM_ANDROID
+#if USE_ANDROID_JNI
 JNI_METHOD void Java_com_epicgames_ue4_GameActivity_nativeAppOpenedWithLocalNotification(JNIEnv* jenv, jobject thiz, jstring jactivationEvent, int32 jFireDate)
 {
 	auto ActivationEvent = FJavaHelper::FStringFromParam(jenv, jactivationEvent);
@@ -77,31 +77,34 @@ FAndroidLocalNotificationService::FAndroidLocalNotificationService(){
 
 void FAndroidLocalNotificationService::ClearAllLocalNotifications()
 {
-	#if PLATFORM_ANDROID
-		extern void AndroidThunkCpp_ClearAllLocalNotifications();
-		AndroidThunkCpp_ClearAllLocalNotifications();
-	#endif
+#if USE_ANDROID_JNI
+	extern void AndroidThunkCpp_ClearAllLocalNotifications();
+	AndroidThunkCpp_ClearAllLocalNotifications();
+#endif
 }
 
-void FAndroidLocalNotificationService::ScheduleLocalNotificationAtTime(const FDateTime& FireDateTime, bool LocalTime, const FText& Title, const FText& Body, const FText& Action, const FString& ActivationEvent)
+int32 FAndroidLocalNotificationService::ScheduleLocalNotificationAtTime(const FDateTime& FireDateTime, bool LocalTime, const FText& Title, const FText& Body, const FText& Action, const FString& ActivationEvent)
 {
-	#if PLATFORM_ANDROID
-		extern void AndroidThunkCpp_ScheduleLocalNotificationAtTime(const FDateTime& FireDateTime, bool LocalTime, const FText& Title, const FText& Body, const FText& Action, const FString& ActivationEvent);
-		AndroidThunkCpp_ScheduleLocalNotificationAtTime(FireDateTime , LocalTime, Title, Body, Action, ActivationEvent);
-	#endif
+#if USE_ANDROID_JNI
+	extern int32 AndroidThunkCpp_ScheduleLocalNotificationAtTime(const FDateTime& FireDateTime, bool LocalTime, const FText& Title, const FText& Body, const FText& Action, const FString& ActivationEvent);
+	return AndroidThunkCpp_ScheduleLocalNotificationAtTime(FireDateTime , LocalTime, Title, Body, Action, ActivationEvent);
+#else
+	return -1;
+#endif
 }
 
-void FAndroidLocalNotificationService::ScheduleLocalNotificationBadgeAtTime(const FDateTime& FireDateTime, bool LocalTime, const FString& ActivationEvent)
+int32 FAndroidLocalNotificationService::ScheduleLocalNotificationBadgeAtTime(const FDateTime& FireDateTime, bool LocalTime, const FString& ActivationEvent)
 {
 	// Do nothing...
+	return -1;
 }
 
 void FAndroidLocalNotificationService::GetLaunchNotification(bool& NotificationLaunchedApp, FString& ActivationEvent, int32& FireDate)
 {
-	#if PLATFORM_ANDROID
-		extern void AndroidThunkCpp_GetLaunchNotification(bool& NotificationLaunchedApp, FString& ActivationEvent, int32& FireDate);
-		AndroidThunkCpp_GetLaunchNotification(NotificationLaunchedApp, ActivationEvent, FireDate);
-	#endif
+#if USE_ANDROID_JNI
+	extern void AndroidThunkCpp_GetLaunchNotification(bool& NotificationLaunchedApp, FString& ActivationEvent, int32& FireDate);
+	AndroidThunkCpp_GetLaunchNotification(NotificationLaunchedApp, ActivationEvent, FireDate);
+#endif
 }
 
 void FAndroidLocalNotificationService::SetLaunchNotification(FString const& ActivationEvent, int32 FireDate)
@@ -113,3 +116,16 @@ void FAndroidLocalNotificationService::CancelLocalNotification(const FString& Ac
 {
 	// TODO
 }
+
+void FAndroidLocalNotificationService::CancelLocalNotification(int32 NotificationId)
+{
+	if (NotificationId < 0)
+	{
+		return;
+	}
+	
+#if USE_ANDROID_JNI
+	extern bool AndroidThunkCpp_DestroyScheduledNotificationIfExists(int32 NotificationId);
+	bool AndroidThunkCpp_DestroyScheduledNotificationIfExists(int32 NotificationId);
+#endif
+}
diff --git a/Engine/Source/Runtime/Android/AndroidLocalNotification/Public/AndroidLocalNotification.h b/Engine/Source/Runtime/Android/AndroidLocalNotification/Public/AndroidLocalNotification.h
--- a/Engine/Source/Runtime/Android/AndroidLocalNotification/Public/AndroidLocalNotification.h
+++ b/Engine/Source/Runtime/Android/AndroidLocalNotification/Public/AndroidLocalNotification.h
@@ -24,43 +24,14 @@ class FAndroidLocalNotificationService : public ILocalNotificationService
 {
 public:
 	FAndroidLocalNotificationService();
-	/** Clear all pending local notifications. Typically this will be done before scheduling new notifications when going into the background */
-	virtual void ClearAllLocalNotifications();
-
-	/** Schedule a local notification at a specific time, inLocalTime specifies the current local time or if UTC time should be used 
-	 * @param FireDateTime The time at which to fire the local notification
-	 * @param LocalTime If true the provided time is in the local timezone, if false it is in UTC
-	 * @param Title The title of the notification
-	 * @param Body The more detailed description of the notification
-	 * @param Action The text to be displayed on the slider controller
-	 * @param ActivationEvent A string that is passed in the delegate callback when the app is brought into the foreground from the user activating the notification
-	*/
-	virtual void ScheduleLocalNotificationAtTime(const FDateTime& FireDateTime, bool LocalTime, const FText& Title, const FText& Body, const FText& Action, const FString& ActivationEvent);
-	
-	/** Schedule a local notification badge at a specific time, inLocalTime specifies the current local time or if UTC time should be used
-	 * @param FireDateTime The time at which to fire the local notification
-	 * @param LocalTime If true the provided time is in the local timezone, if false it is in UTC
-	 * @param ActivationEvent A string that is passed in the delegate callback when the app is brought into the foreground from the user activating the notification
-	*/
-	virtual void ScheduleLocalNotificationBadgeAtTime(const FDateTime& FireDateTime, bool LocalTime, const FString& ActivationEvent);
 
-	/** Get the local notification that was used to launch the app
-	 * @param NotificationLaunchedApp Return true if a notification was used to launch the app
-	 * @param ActivationEvent Returns the name of the ActivationEvent if a notification was used to launch the app
-	 * @param FireDate Returns the time the notification was activated
-	*/
+	virtual void ClearAllLocalNotifications();
+	virtual int32 ScheduleLocalNotificationAtTime(const FDateTime& FireDateTime, bool LocalTime, const FText& Title, const FText& Body, const FText& Action, const FString& ActivationEvent);
+	virtual int32 ScheduleLocalNotificationBadgeAtTime(const FDateTime& FireDateTime, bool LocalTime, const FString& ActivationEvent);
 	virtual void GetLaunchNotification(bool& NotificationLaunchedApp, FString& ActivationEvent, int32& FireDate);
-
-	/** Set the local notification that was used to launch the app
-	 * @param ActivationEvent Returns the name of the ActivationEvent if a notification was used to launch the app
-	 * @param FireDate Returns the time the notification was activated
-	*/
 	virtual void SetLaunchNotification(FString const& ActivationEvent, int32 FireDate);
-
-	/** Cancel a local notification given the ActivationEvent
-	 * @param ActivationEvent The string passed into the Schedule call for the notification to be cancelled
-	*/
 	virtual void CancelLocalNotification(const FString& ActivationEvent);
+	virtual void CancelLocalNotification(int NotificationId);
 
 private:
 	bool	AppLaunchedWithNotification;
diff --git a/Engine/Source/Runtime/Engine/Classes/Kismet/BlueprintPlatformLibrary.h b/Engine/Source/Runtime/Engine/Classes/Kismet/BlueprintPlatformLibrary.h
--- a/Engine/Source/Runtime/Engine/Classes/Kismet/BlueprintPlatformLibrary.h
+++ b/Engine/Source/Runtime/Engine/Classes/Kismet/BlueprintPlatformLibrary.h
@@ -182,7 +182,7 @@ public:
 	 * @param ActivationEvent A string that is passed in the delegate callback when the app is brought into the foreground from the user activating the notification
 	*/
 	UFUNCTION(BlueprintCallable, Category="Platform|LocalNotification")
-	static void ScheduleLocalNotificationAtTime(const FDateTime& FireDateTime, bool LocalTime, const FText& Title, const FText& Body, const FText& Action, const FString& ActivationEvent);
+	static int32 ScheduleLocalNotificationAtTime(const FDateTime& FireDateTime, bool LocalTime, const FText& Title, const FText& Body, const FText& Action, const FString& ActivationEvent);
 
 	/** Schedule a local notification to fire inSecondsFromNow from now 
 	 * @param inSecondsFromNow The seconds until the notification should fire
@@ -193,7 +193,7 @@ public:
 	 * @param ActivationEvent A string that is passed in the delegate callback when the app is brought into the foreground from the user activating the notification
 	*/
 	UFUNCTION(BlueprintCallable, Category="Platform|LocalNotification")
-	static void ScheduleLocalNotificationFromNow(int32 inSecondsFromNow, const FText& Title, const FText& Body, const FText& Action, const FString& ActivationEvent);
+	static int32 ScheduleLocalNotificationFromNow(int32 inSecondsFromNow, const FText& Title, const FText& Body, const FText& Action, const FString& ActivationEvent);
 
 	/** Schedule a local notification badge at a specific time, inLocalTime specifies the current local time or if UTC time should be used
 	 * @param FireDateTime The time at which to fire the local notification
@@ -201,7 +201,7 @@ public:
 	 * @param ActivationEvent A string that is passed in the delegate callback when the app is brought into the foreground from the user activating the notification
 	*/
 	UFUNCTION(BlueprintCallable, Category = "Platform|LocalNotification")
-	static void ScheduleLocalNotificationBadgeAtTime(const FDateTime& FireDateTime, bool LocalTime, const FString& ActivationEvent);
+	static int32 ScheduleLocalNotificationBadgeAtTime(const FDateTime& FireDateTime, bool LocalTime, const FString& ActivationEvent);
 	
 	/** Schedule a local notification badge to fire inSecondsFromNow from now
 	 * @param inSecondsFromNow The seconds until the notification should fire
@@ -216,6 +216,12 @@ public:
 	UFUNCTION(BlueprintCallable, Category="Platform|LocalNotification")
 	static void CancelLocalNotification(const FString& ActivationEvent);
 
+	/** Cancel a local notification given the ActivationEvent
+	 * @param NotificationId The Id returned from one of the ScheduleLocalNotification* functions
+	 */
+	UFUNCTION(BlueprintCallable, Category="Platform|LocalNotification")
+	static void CancelLocalNotificationById(int32 NotificationId);
+
 	/** Get the local notification that was used to launch the app
 	 * @param NotificationLaunchedApp Return true if a notification was used to launch the app
 	 * @param ActivationEvent Returns the name of the ActivationEvent if a notification was used to launch the app
diff --git a/Engine/Source/Runtime/Engine/Private/Kismet/BlueprintPlatformLibrary.cpp b/Engine/Source/Runtime/Engine/Private/Kismet/BlueprintPlatformLibrary.cpp
--- a/Engine/Source/Runtime/Engine/Private/Kismet/BlueprintPlatformLibrary.cpp
+++ b/Engine/Source/Runtime/Engine/Private/Kismet/BlueprintPlatformLibrary.cpp
@@ -96,38 +96,38 @@ void UBlueprintPlatformLibrary::ClearAllLocalNotifications()
 	platformService->ClearAllLocalNotifications();
 }
 
-void UBlueprintPlatformLibrary::ScheduleLocalNotificationAtTime(const FDateTime& FireDateTime, bool inLocalTime, const FText& Title, const FText& Body, const FText& Action, const FString& ActivationEvent)
+int32 UBlueprintPlatformLibrary::ScheduleLocalNotificationAtTime(const FDateTime& FireDateTime, bool inLocalTime, const FText& Title, const FText& Body, const FText& Action, const FString& ActivationEvent)
 {
 	if(platformService == nullptr)
 	{
 		UE_LOG(LogBlueprintUserMessages, Warning, TEXT("ScheduleLocalNotificationAtTime(): No local notification service"));
-		return;
+		return -1;
 	}
 
 	UE_LOG(LogBlueprintUserMessages, Log, TEXT("Scheduling notification %s at %d/%d/%d %d:%d:%d %s"), *(Title.ToString()), FireDateTime.GetMonth(), FireDateTime.GetDay(), FireDateTime.GetYear(), FireDateTime.GetHour(), FireDateTime.GetMinute(), FireDateTime.GetSecond(), inLocalTime ? TEXT("Local") : TEXT("UTC"));
 	
-	platformService->ScheduleLocalNotificationAtTime(FireDateTime, inLocalTime, Title, Body, Action, ActivationEvent);
+	return platformService->ScheduleLocalNotificationAtTime(FireDateTime, inLocalTime, Title, Body, Action, ActivationEvent);
 }
        
-void UBlueprintPlatformLibrary::ScheduleLocalNotificationFromNow(int32 inSecondsFromNow, const FText& Title, const FText& Body, const FText& Action, const FString& ActivationEvent)
+int32 UBlueprintPlatformLibrary::ScheduleLocalNotificationFromNow(int32 inSecondsFromNow, const FText& Title, const FText& Body, const FText& Action, const FString& ActivationEvent)
 {
 	FDateTime TargetTime = FDateTime::Now();
 	TargetTime += FTimespan::FromSeconds(inSecondsFromNow);
 
-	ScheduleLocalNotificationAtTime(TargetTime, true, Title, Body, Action, ActivationEvent);
+	return ScheduleLocalNotificationAtTime(TargetTime, true, Title, Body, Action, ActivationEvent);
 }
 
-void UBlueprintPlatformLibrary::ScheduleLocalNotificationBadgeAtTime(const FDateTime& FireDateTime, bool inLocalTime, const FString& ActivationEvent)
+int32 UBlueprintPlatformLibrary::ScheduleLocalNotificationBadgeAtTime(const FDateTime& FireDateTime, bool inLocalTime, const FString& ActivationEvent)
 {
 	if (platformService == nullptr)
 	{
 		UE_LOG(LogBlueprintUserMessages, Warning, TEXT("ScheduleLocalNotificationBadgeAtTime(): No local notification service"));
-		return;
+		return -1;
 	}
 
 	UE_LOG(LogBlueprintUserMessages, Log, TEXT("Scheduling notification badge %s at %d/%d/%d %d:%d:%d %s"), *ActivationEvent, FireDateTime.GetMonth(), FireDateTime.GetDay(), FireDateTime.GetYear(), FireDateTime.GetHour(), FireDateTime.GetMinute(), FireDateTime.GetSecond(), inLocalTime ? TEXT("Local") : TEXT("UTC"));
 
-	platformService->ScheduleLocalNotificationBadgeAtTime(FireDateTime, inLocalTime, ActivationEvent);
+	return platformService->ScheduleLocalNotificationBadgeAtTime(FireDateTime, inLocalTime, ActivationEvent);
 }
 
 void UBlueprintPlatformLibrary::ScheduleLocalNotificationBadgeFromNow(int32 inSecondsFromNow, const FString& ActivationEvent)
@@ -152,6 +152,20 @@ void UBlueprintPlatformLibrary::CancelLocalNotification(const FString& Activatio
 	platformService->CancelLocalNotification(ActivationEvent);
 }
 
+UFUNCTION(BlueprintCallable, Category="Platform|LocalNotification")
+void UBlueprintPlatformLibrary::CancelLocalNotificationById(int32 NotificationId)
+{
+	if(platformService == nullptr)
+	{
+		UE_LOG(LogBlueprintUserMessages, Warning, TEXT("CancelLocalNotification(): No local notification service"));
+		return;
+	}
+	
+	UE_LOG(LogBlueprintUserMessages, Log, TEXT("Canceling notification %f"), NotificationId);
+	
+	platformService->CancelLocalNotification(NotificationId);
+}
+
 void UBlueprintPlatformLibrary::GetLaunchNotification(bool& NotificationLaunchedApp, FString& ActivationEvent, int32& FireDate)
 {
 	if(platformService == nullptr)
diff --git a/Engine/Source/Runtime/Engine/Public/LocalNotification.h b/Engine/Source/Runtime/Engine/Public/LocalNotification.h
--- a/Engine/Source/Runtime/Engine/Public/LocalNotification.h
+++ b/Engine/Source/Runtime/Engine/Public/LocalNotification.h
@@ -25,15 +25,17 @@ public:
 	 * @param Body The more detailed description of the notification
 	 * @param Action The text to be displayed on the slider controller
 	 * @param ActivationEvent A string that is passed in the delegate callback when the app is brought into the foreground from the user activating the notification
+	 * @return A unique identifier used to cancel the notification
 	*/
-	virtual void ScheduleLocalNotificationAtTime(const FDateTime& FireDateTime, bool LocalTime, const FText& Title, const FText& Body, const FText& Action, const FString& ActivationEvent) = 0;
+	virtual int32 ScheduleLocalNotificationAtTime(const FDateTime& FireDateTime, bool LocalTime, const FText& Title, const FText& Body, const FText& Action, const FString& ActivationEvent) = 0;
 
 	/** Schedule a local notification badge at a specific time, inLocalTime specifies the current local time or if UTC time should be used
 	 * @param FireDateTime The time at which to fire the local notification
 	 * @param LocalTime If true the provided time is in the local timezone, if false it is in UTC
 	 * @param ActivationEvent A string that is passed in the delegate callback when the app is brought into the foreground from the user activating the notification
+	 * @return A unique identifier used to cancel the notification
 	*/
-	virtual void ScheduleLocalNotificationBadgeAtTime(const FDateTime& FireDateTime, bool LocalTime, const FString& ActivationEvent) = 0;
+	virtual int32 ScheduleLocalNotificationBadgeAtTime(const FDateTime& FireDateTime, bool LocalTime, const FString& ActivationEvent) = 0;
 
 	/** Get the local notification that was used to launch the app
 	 * @param NotificationLaunchedApp Return true if a notification was used to launch the app
@@ -47,6 +49,11 @@ public:
 	*/
 	virtual void CancelLocalNotification(const FString& ActivationEvent) = 0;
 
+	/** Cancel a local notification given the ActivationEvent
+	 * @param NotificationId The Id returned from ScheduleLocalNotificationAtTime*
+	 */
+	virtual void CancelLocalNotification(int32 NotificationId) = 0;
+
 
 	/** Set the local notification that was used to launch the app
 	 * @param ActivationEvent Returns the name of the ActivationEvent if a notification was used to launch the app
diff --git a/Engine/Source/Runtime/IOS/IOSLocalNotification/Private/IOSLocalNotification.cpp b/Engine/Source/Runtime/IOS/IOSLocalNotification/Private/IOSLocalNotification.cpp
--- a/Engine/Source/Runtime/IOS/IOSLocalNotification/Private/IOSLocalNotification.cpp
+++ b/Engine/Source/Runtime/IOS/IOSLocalNotification/Private/IOSLocalNotification.cpp
@@ -16,9 +16,7 @@
 #include "Modules/ModuleManager.h"
 #include "Logging/LogMacros.h"
 
-#if __IPHONE_OS_VERSION_MIN_REQUIRED >= __IPHONE_10_0
 #import <UserNotifications/UserNotifications.h>
-#endif
 
 DEFINE_LOG_CATEGORY(LogIOSLocalNotification);
 
@@ -40,7 +38,6 @@ public:
 	}
 
 #if !PLATFORM_TVOS
-#if __IPHONE_OS_VERSION_MIN_REQUIRED >= __IPHONE_10_0
 	static UNMutableNotificationContent* CreateNotificationContent(const FText& Title, const FText& Body, const FText& Action, const FString& ActivationEvent, uint32 BadgeNumber)
 	{
 		UNMutableNotificationContent* Content = [UNMutableNotificationContent new];
@@ -100,47 +97,6 @@ public:
 		
 		return trigger;
 	}
-#else
-	static UILocalNotification* CreateLocalNotification(const FDateTime& FireDateTime, bool bLocalTime, const FString& ActivationEvent)
-	{
-		UIApplication* application = [UIApplication sharedApplication];
-
-		NSCalendar *calendar = [NSCalendar autoupdatingCurrentCalendar];
-		NSDateComponents *dateComps = [[NSDateComponents alloc] init];
-		[dateComps setDay : FireDateTime.GetDay()];
-		[dateComps setMonth : FireDateTime.GetMonth()];
-		[dateComps setYear : FireDateTime.GetYear()];
-		[dateComps setHour : FireDateTime.GetHour()];
-		[dateComps setMinute : FireDateTime.GetMinute()];
-		[dateComps setSecond : FireDateTime.GetSecond()];
-		NSDate *itemDate = [calendar dateFromComponents : dateComps];
-
-		UILocalNotification *localNotif = [[UILocalNotification alloc] init];
-		if (localNotif != nil)
-		{
-			localNotif.fireDate = itemDate;
-			if (bLocalTime)
-			{
-				localNotif.timeZone = [NSTimeZone defaultTimeZone];
-			}
-			else
-			{
-				localNotif.timeZone = nil;
-			}
-
-			NSString* activateEventNSString = [NSString stringWithFString:ActivationEvent];
-			if (activateEventNSString != nil)
-			{
-				NSDictionary* infoDict = [NSDictionary dictionaryWithObject:activateEventNSString forKey:@"ActivationEvent"];
-				if (infoDict != nil)
-				{
-					localNotif.userInfo = infoDict;
-				}
-			}
-		}
-		return localNotif;
-	}
-#endif
 #endif
 };
 
@@ -158,125 +114,69 @@ FIOSLocalNotificationService::FIOSLocalNotificationService()
 void FIOSLocalNotificationService::ClearAllLocalNotifications()
 {
 #if !PLATFORM_TVOS
-#if __IPHONE_OS_VERSION_MIN_REQUIRED >= __IPHONE_10_0
 	UNUserNotificationCenter *Center = [UNUserNotificationCenter currentNotificationCenter];
 	
 	[Center removeAllPendingNotificationRequests];
-#else
-	UIApplication* application = [UIApplication sharedApplication];
-	
-	[application cancelAllLocalNotifications];
-#endif
 #endif
 }
 
-static uint32 NotificationNumber = 0;
-void FIOSLocalNotificationService::ScheduleLocalNotificationAtTime(const FDateTime& FireDateTime, bool LocalTime, const FText& Title, const FText& Body, const FText& Action, const FString& ActivationEvent)
+static int32 NotificationNumber = 0;
+int32 FIOSLocalNotificationService::ScheduleLocalNotificationAtTime(const FDateTime& FireDateTime, bool LocalTime, const FText& Title, const FText& Body, const FText& Action, const FString& ActivationEvent)
 {
 #if !PLATFORM_TVOS
-#if __IPHONE_OS_VERSION_MIN_REQUIRED >= __IPHONE_10_0
-    
     //Create local copies of these for the block to capture
     FDateTime FireDateTimeCopy = FireDateTime;
     FText TitleCopy = Title;
     FText BodyCopy = Body;
     FText ActionCopy = Action;
     FString ActivationEventCopy = ActivationEvent;
-    
+	int32 CurrentNotificationId = NotificationNumber++;
+	
 	//have to schedule notification on main thread queue
 	dispatch_async(dispatch_get_main_queue(), ^{
         UNMutableNotificationContent* Content = FIOSLocalNotificationModule::CreateNotificationContent(TitleCopy, BodyCopy, ActionCopy, ActivationEventCopy, 1);
         UNCalendarNotificationTrigger* Trigger = FIOSLocalNotificationModule::CreateCalendarNotificationTrigger(FireDateTimeCopy);
         
-        UNUserNotificationCenter *Center = [UNUserNotificationCenter currentNotificationCenter];
-        
-        FString NotId = TitleCopy.ToString() + FString::FromInt(NotificationNumber++);
-        NSString* NotificationIdentifier = [NSString stringWithFString : NotId];
-		
-        UNNotificationRequest *Request = [UNNotificationRequest requestWithIdentifier:NotificationIdentifier content:Content trigger:Trigger];
+        UNNotificationRequest *Request = [UNNotificationRequest requestWithIdentifier:@(CurrentNotificationId).stringValue content:Content trigger:Trigger];
 
+		UNUserNotificationCenter *Center = [UNUserNotificationCenter currentNotificationCenter];
         [Center addNotificationRequest : Request withCompletionHandler : ^ (NSError * _Nullable error) {
             if (error != nil) {
-                UE_LOG(LogIOSLocalNotification, Warning, TEXT("Error scheduling notification: %s"), *NotId);
+                UE_LOG(LogIOSLocalNotification, Warning, TEXT("Error scheduling notification: %d"), CurrentNotificationId);
             }
         }];
     });
+	return CurrentNotificationId;
 #else
-    
-    //Create local copies of these for the block to capture
-    FDateTime FireDateTimeCopy = FireDateTime;
-    FText TitleCopy = Title;
-    FText BodyCopy = Body;
-    FText ActionCopy = Action;
-    FString ActivationEventCopy = ActivationEvent;
-    
-	//have to schedule notification on main thread queue
-	dispatch_async(dispatch_get_main_queue(), ^{
-		UILocalNotification *localNotif = FIOSLocalNotificationModule::CreateLocalNotification(FireDateTimeCopy, LocalTime, ActivationEventCopy);
-		if (localNotif == nil)
-			return;
-
-		NSString*	alertBody = [NSString stringWithFString : BodyCopy.ToString()];
-		if (alertBody != nil)
-		{
-			localNotif.alertBody = alertBody;
-		}
-
-		NSString*	alertAction = [NSString stringWithFString : ActionCopy.ToString()];
-		if (alertAction != nil)
-		{
-			localNotif.alertAction = alertAction;
-		}
-
-		if ([IOSAppDelegate GetDelegate].OSVersion >= 8.2f)
-		{
-			NSString*	alertTitle = [NSString stringWithFString : TitleCopy.ToString()];
-			if (alertTitle != nil)
-			{
-				localNotif.alertTitle = alertTitle;
-			}
-		}
-
-		localNotif.soundName = UILocalNotificationDefaultSoundName;
-		localNotif.applicationIconBadgeNumber = 1;
-
-		[[UIApplication sharedApplication] scheduleLocalNotification:localNotif];
-    });
-#endif
+	return -1;
 #endif
 }
 
-void FIOSLocalNotificationService::ScheduleLocalNotificationBadgeAtTime(const FDateTime& FireDateTime, bool LocalTime, const FString& ActivationEvent)
+int32 FIOSLocalNotificationService::ScheduleLocalNotificationBadgeAtTime(const FDateTime& FireDateTime, bool LocalTime, const FString& ActivationEvent)
 {
 #if !PLATFORM_TVOS
-#if __IPHONE_OS_VERSION_MIN_REQUIRED >= __IPHONE_10_0
-	UNMutableNotificationContent* Content = FIOSLocalNotificationModule::CreateNotificationContent(FText(), FText(), FText(), ActivationEvent, 1);
-	UNCalendarNotificationTrigger* Trigger = FIOSLocalNotificationModule::CreateCalendarNotificationTrigger(FireDateTime);
-	
-	UNUserNotificationCenter *Center = [UNUserNotificationCenter currentNotificationCenter];
+	FDateTime FireDateTimeCopy = FireDateTime;
+	FString ActivationEventCopy = ActivationEvent;
+	int32 CurrentNotificationId = NotificationNumber++;
 	
-	FString NotId = FString(TEXT("Badge")) + FString::FromInt(NotificationNumber++);
-	NSString* NotificationIdentifier = [NSString stringWithFString:NotId];
-	
-	UNNotificationRequest *Request = [UNNotificationRequest requestWithIdentifier:NotificationIdentifier content:Content trigger:Trigger];
+	//have to schedule notification on main thread queue
+	dispatch_async(dispatch_get_main_queue(), ^{
+		UNMutableNotificationContent* Content = FIOSLocalNotificationModule::CreateNotificationContent(FText(), FText(), FText(), ActivationEventCopy, 1);
+		UNCalendarNotificationTrigger* Trigger = FIOSLocalNotificationModule::CreateCalendarNotificationTrigger(FireDateTime);
+		
+		UNNotificationRequest *Request = [UNNotificationRequest requestWithIdentifier:@(CurrentNotificationId).stringValue content:Content trigger:Trigger];
+
+		UNUserNotificationCenter *Center = [UNUserNotificationCenter currentNotificationCenter];
+		[Center addNotificationRequest:Request withCompletionHandler:^(NSError * _Nullable error) {
+			if (error != nil) {
+				UE_LOG(LogIOSLocalNotification, Warning, TEXT("Error scheduling notification: %d"), CurrentNotificationId);
+			}
+		}];
+	});
 	
-	[Center addNotificationRequest:Request withCompletionHandler:^(NSError * _Nullable error) {
-		if (error != nil) {
-			UE_LOG(LogIOSLocalNotification, Warning, TEXT("Error scheduling notification: %s"), *NotId);
-		}
-	}];
+	return CurrentNotificationId;
 #else
-	UILocalNotification *localNotif = FIOSLocalNotificationModule::CreateLocalNotification(FireDateTime, LocalTime, ActivationEvent);
-	if (localNotif == nil)
-		return;
-
-	// As per Apple documentation, a nil 'alertBody' results in 'no alert'
-	// https://developer.apple.com/reference/uikit/uilocalnotification/1616646-alertbody?language=objc
-	localNotif.alertBody = nil;
-	localNotif.applicationIconBadgeNumber = 1;
-
-	[[UIApplication sharedApplication] scheduleLocalNotification:localNotif];
-#endif
+	return -1;
 #endif
 }
 
@@ -285,6 +185,14 @@ void FIOSLocalNotificationService::CancelLocalNotification(const FString& Activa
 	// TODO
 }
 
+void FIOSLocalNotificationService::CancelLocalNotification(int32 NotificationId)
+{
+#if !PLATFORM_TVOS
+	UNUserNotificationCenter *Center = [UNUserNotificationCenter currentNotificationCenter];
+	[Center removePendingNotificationRequestsWithIdentifiers:@[@(NotificationId).stringValue]];
+#endif
+}
+
 void FIOSLocalNotificationService::GetLaunchNotification(bool& NotificationLaunchedApp, FString& ActivationEvent, int32& FireDate)
 {
 	NotificationLaunchedApp = AppLaunchedWithNotification;
@@ -303,16 +211,15 @@ void FIOSLocalNotificationService::SetLaunchNotification(FString const& Activati
 static FIOSLocalNotificationService::FAllowedNotifications NotificationsAllowedDelegate;
 void FIOSLocalNotificationService::CheckAllowedNotifications(const FAllowedNotifications& AllowedNotificationsDelegate)
 {
-#if __IPHONE_OS_VERSION_MIN_REQUIRED >= __IPHONE_10_0
 	NotificationsAllowedDelegate = AllowedNotificationsDelegate;
 	
+#if !PLATFORM_TVOS
 	UNUserNotificationCenter *Center = [UNUserNotificationCenter currentNotificationCenter];
 	[Center getNotificationSettingsWithCompletionHandler:^(UNNotificationSettings * _Nonnull settings) {
 		bool NotificationsAllowed = settings.authorizationStatus == UNAuthorizationStatusAuthorized;
 		NotificationsAllowedDelegate.ExecuteIfBound(NotificationsAllowed);
 	}];
-	
 #else
-	checkf(false, TEXT("For min iOS version < 10 use FIOSPlatformMisc::IsAllowedRemoteNotifications()."));
+	NotificationsAllowedDelegate.ExecuteIfBound(false);
 #endif
 }
diff --git a/Engine/Source/Runtime/IOS/IOSLocalNotification/Public/IOSLocalNotification.h b/Engine/Source/Runtime/IOS/IOSLocalNotification/Public/IOSLocalNotification.h
--- a/Engine/Source/Runtime/IOS/IOSLocalNotification/Public/IOSLocalNotification.h
+++ b/Engine/Source/Runtime/IOS/IOSLocalNotification/Public/IOSLocalNotification.h
@@ -20,44 +20,13 @@ class FIOSLocalNotificationService : public ILocalNotificationService
 public:
 	FIOSLocalNotificationService();
 	virtual ~FIOSLocalNotificationService() {}
-
-	/** Clear all pending local notifications. Typically this will be done before scheduling new notifications when going into the background */
 	virtual void ClearAllLocalNotifications();
-
-	/** Schedule a local notification at a specific time, inLocalTime specifies the current local time or if UTC time should be used 
-	 * @param FireDateTime The time at which to fire the local notification
-	 * @param LocalTime If true the provided time is in the local timezone, if false it is in UTC
-	 * @param Title The title of the notification
-	 * @param Body The more detailed description of the notification
-	 * @param Action The text to be displayed on the slider controller
-	 * @param ActivationEvent A string that is passed in the delegate callback when the app is brought into the foreground from the user activating the notification
-	*/
-	virtual void ScheduleLocalNotificationAtTime(const FDateTime& FireDateTime, bool LocalTime, const FText& Title, const FText& Body, const FText& Action, const FString& ActivationEvent);
-	
-	/** Schedule a local notification badge at a specific time, inLocalTime specifies the current local time or if UTC time should be used
-	* @param FireDateTime The time at which to fire the local notification
-	* @param LocalTime If true the provided time is in the local timezone, if false it is in UTC
-	* @param ActivationEvent A string that is passed in the delegate callback when the app is brought into the foreground from the user activating the notification
-	*/
-	virtual void ScheduleLocalNotificationBadgeAtTime(const FDateTime& FireDateTime, bool LocalTime, const FString& ActivationEvent);
-
-	/** Get the local notification that was used to launch the app
-	 * @param NotificationLaunchedApp Return true if a notification was used to launch the app
-	 * @param ActivationEvent Returns the name of the ActivationEvent if a notification was used to launch the app
-	 * @param FireDate Returns the time the notification was activated
-	*/
+	virtual int32 ScheduleLocalNotificationAtTime(const FDateTime& FireDateTime, bool LocalTime, const FText& Title, const FText& Body, const FText& Action, const FString& ActivationEvent);
+	virtual int32 ScheduleLocalNotificationBadgeAtTime(const FDateTime& FireDateTime, bool LocalTime, const FString& ActivationEvent);
 	virtual void GetLaunchNotification(bool& NotificationLaunchedApp, FString& ActivationEvent, int32& FireDate);
-
-	/** Set the local notification that was used to launch the app[REDACTED_USER_HOME]/pfEpicContracting/Dev-Platform/Engine/Source/Runtime/IOS/IOSLocalNotification/IOSLocalNotification.Build.cs
-	 * @param ActivationEvent Returns the name of the ActivationEvent if a notification was used to launch the app
-	 * @param FireDate Returns the time the notification was activated
-	*/
 	virtual void SetLaunchNotification(FString const& ActivationEvent, int32 FireDate);
-
-	/** Cancel a local notification given the ActivationEvent
-	 * @param ActivationEvent The string passed into the Schedule call for the notification to be cancelled
-	*/
 	virtual void CancelLocalNotification(const FString& ActivationEvent);
+	virtual void CancelLocalNotification(int NotificationId);
 
 	// Check if notifications are allowed if min iOS version is > 10
 	DECLARE_DELEGATE_OneParam(FAllowedNotifications, bool);
diff --git a/Engine/Source/Runtime/Launch/Private/Android/AndroidJNI.cpp b/Engine/Source/Runtime/Launch/Private/Android/AndroidJNI.cpp
--- a/Engine/Source/Runtime/Launch/Private/Android/AndroidJNI.cpp
+++ b/Engine/Source/Runtime/Launch/Private/Android/AndroidJNI.cpp
@@ -101,7 +101,7 @@ void FJavaWrapper::FindClassesAndMethods(JNIEnv* Env)
 	AndroidThunkJava_LocalNotificationScheduleAtTime = FindMethod(Env, GameActivityClassID, "AndroidThunkJava_LocalNotificationScheduleAtTime", "(Ljava/lang/String;ZLjava/lang/String;Ljava/lang/String;Ljava/lang/String;Ljava/lang/String;)V", bIsOptional);
 	AndroidThunkJava_LocalNotificationClearAll = FindMethod(Env, GameActivityClassID, "AndroidThunkJava_LocalNotificationClearAll", "()V", bIsOptional);
 	AndroidThunkJava_LocalNotificationGetLaunchNotification = FindMethod(Env, GameActivityClassID, "AndroidThunkJava_LocalNotificationGetLaunchNotification", "()Lcom/epicgames/ue4/GameActivity$LaunchNotification;", bIsOptional);
-	//AndroidThunkJava_LocalNotificationDestroyIfExists = FindMethod(Env, GameActivityClassID, "AndroidThunkJava_LocalNotificationDestroyIfExists", "(I)Z", bIsOptional);
+	AndroidThunkJava_LocalNotificationDestroyIfExists = FindMethod(Env, GameActivityClassID, "AndroidThunkJava_LocalNotificationDestroyIfExists", "(I)Z", bIsOptional);
 	AndroidThunkJava_GetNetworkConnectionType = FindMethod(Env, GameActivityClassID, "AndroidThunkJava_GetNetworkConnectionType", "()I", bIsOptional);
 	AndroidThunkJava_GetAndroidId = FindMethod(Env, GameActivityClassID, "AndroidThunkJava_GetAndroidId", "()Ljava/lang/String;", bIsOptional);
 	AndroidThunkJava_ShareURL = FindMethod(Env, GameActivityClassID, "AndroidThunkJava_ShareURL", "(Ljava/lang/String;Ljava/lang/String;Ljava/lang/String;II)V", bIsOptional);
@@ -378,7 +378,7 @@ jmethodID FJavaWrapper::AndroidThunkJava_ShowHiddenAlertDialog;
 jmethodID FJavaWrapper::AndroidThunkJava_LocalNotificationScheduleAtTime;
 jmethodID FJavaWrapper::AndroidThunkJava_LocalNotificationClearAll;
 jmethodID FJavaWrapper::AndroidThunkJava_LocalNotificationGetLaunchNotification;
-//jmethodID FJavaWrapper::AndroidThunkJava_LocalNotificationDestroyIfExists;
+jmethodID FJavaWrapper::AndroidThunkJava_LocalNotificationDestroyIfExists;
 jmethodID FJavaWrapper::AndroidThunkJava_GetNetworkConnectionType;
 jmethodID FJavaWrapper::AndroidThunkJava_GetAndroidId;
 jmethodID FJavaWrapper::AndroidThunkJava_ShareURL;
@@ -1417,7 +1417,7 @@ void AndroidThunkCpp_SetDesiredViewSize(int32 Width, int32 Height)
 // #endif
 }
 
-void AndroidThunkCpp_ScheduleLocalNotificationAtTime(const FDateTime& FireDateTime, bool LocalTime, const FText& Title, const FText& Body, const FText& Action, const FString& ActivationEvent) 
+int32 AndroidThunkCpp_ScheduleLocalNotificationAtTime(const FDateTime& FireDateTime, bool LocalTime, const FText& Title, const FText& Body, const FText& Action, const FString& ActivationEvent)
 {
 	//Convert FireDateTime to yyyy-MM-dd HH:mm:ss in order to pass to java
 	FString FireDateTimeFormatted = FString::FromInt(FireDateTime.GetYear()) + "-" + FString::FromInt(FireDateTime.GetMonth()) + "-" + FString::FromInt(FireDateTime.GetDay()) + " " + FString::FromInt(FireDateTime.GetHour()) + ":" + FString::FromInt(FireDateTime.GetMinute()) + ":" + FString::FromInt(FireDateTime.GetSecond());
@@ -1431,8 +1431,10 @@ void AndroidThunkCpp_ScheduleLocalNotificationAtTime(const FDateTime& FireDateTi
 		auto jAction = FJavaHelper::ToJavaString(Env, Action.ToString());
 		auto jActivationEvent = FJavaHelper::ToJavaString(Env, ActivationEvent);
 		
-		FJavaWrapper::CallVoidMethod(Env, FJavaWrapper::GameActivityThis, FJavaWrapper::AndroidThunkJava_LocalNotificationScheduleAtTime, *jFireDateTime, LocalTime, *jTitle, *jBody, *jAction, *jActivationEvent);
+		return FJavaWrapper::CallIntMethod(Env, FJavaWrapper::GameActivityThis, FJavaWrapper::AndroidThunkJava_LocalNotificationScheduleAtTime, *jFireDateTime, LocalTime, *jTitle, *jBody, *jAction, *jActivationEvent);
 	}
+	
+	return -1;
 }
 
 void AndroidThunkCpp_GetLaunchNotification(bool& NotificationLaunchedApp, FString& ActivationEvent, int32& FireDate)
@@ -1465,7 +1467,7 @@ void AndroidThunkCpp_ClearAllLocalNotifications()
 		FJavaWrapper::CallVoidMethod(Env, FJavaWrapper::GameActivityThis, FJavaWrapper::AndroidThunkJava_LocalNotificationClearAll);
 	}
 }
-/*
+
 bool AndroidThunkCpp_DestroyScheduledNotificationIfExists(int32 NotificationId)
 {
 	bool Result = false;
@@ -1475,7 +1477,6 @@ bool AndroidThunkCpp_DestroyScheduledNotificationIfExists(int32 NotificationId)
 	}
 	return Result;
 }
-*/
 
 int32 AndroidThunkCpp_GetNetworkConnectionType()
 {
diff --git a/Engine/Source/Runtime/Launch/Public/Android/AndroidJNI.h b/Engine/Source/Runtime/Launch/Public/Android/AndroidJNI.h
--- a/Engine/Source/Runtime/Launch/Public/Android/AndroidJNI.h
+++ b/Engine/Source/Runtime/Launch/Public/Android/AndroidJNI.h
@@ -57,7 +57,7 @@ public:
 	static jmethodID AndroidThunkJava_LocalNotificationScheduleAtTime;
 	static jmethodID AndroidThunkJava_LocalNotificationClearAll;
 	static jmethodID AndroidThunkJava_LocalNotificationGetLaunchNotification;
-	//static jmethodID AndroidThunkJava_LocalNotificationDestroyIfExists; - This is not needed yet but will be soon so just leaving commented out for now
+	static jmethodID AndroidThunkJava_LocalNotificationDestroyIfExists;
 	static jmethodID AndroidThunkJava_GetNetworkConnectionType;
 	static jmethodID AndroidThunkJava_GetAndroidId;
 	static jmethodID AndroidThunkJava_ShareURL;
