Fixing iOS from previous check-in

diff --git a/Engine/Source/Runtime/IOS/IOSAudio/Private/IOSAudioDevice.cpp b/Engine/Source/Runtime/IOS/IOSAudio/Private/IOSAudioDevice.cpp
--- a/Engine/Source/Runtime/IOS/IOSAudio/Private/IOSAudioDevice.cpp
+++ b/Engine/Source/Runtime/IOS/IOSAudio/Private/IOSAudioDevice.cpp
@@ -198,7 +198,7 @@ bool FIOSAudioDevice::InitializeHardware()
 		return false;
 	}
 
-	uint32 BusCount = MaxChannels * CHANNELS_PER_BUS;
+	uint32 BusCount = GetMaxChannels() * CHANNELS_PER_BUS;
 	Status = AudioUnitSetProperty(MixerUnit,
 	                              kAudioUnitProperty_ElementCount,
 	                              kAudioUnitScope_Input,
