[NO NON-PII COMMIT MESSAGE]

diff --git a/Engine/Source/Runtime/IOS/IOSAudio/Private/IOSAudioBuffer.cpp b/Engine/Source/Runtime/IOS/IOSAudio/Private/IOSAudioBuffer.cpp
--- a/Engine/Source/Runtime/IOS/IOSAudio/Private/IOSAudioBuffer.cpp
+++ b/Engine/Source/Runtime/IOS/IOSAudio/Private/IOSAudioBuffer.cpp
@@ -31,14 +31,14 @@ FIOSAudioSoundBuffer::FIOSAudioSoundBuffer(FIOSAudioDevice* InAudioDevice, USoun
 {
 	if (!bIsProcedural)
 	{
-		DecompressionState = static_cast<FADPCMAudioInfo*>(InAudioDevice->CreateCompressedAudioInfo(InWave));
+		DecompressionState = InAudioDevice->CreateCompressedAudioInfo(InWave);
 
 		if (!ReadCompressedInfo(InWave))
 		{
 			return;
 		}
 
-		SoundFormat = static_cast<ESoundFormat>(DecompressionState->GetFormatTag());
+		SoundFormat = static_cast<ESoundFormat>(static_cast<FADPCMAudioInfo*>(DecompressionState)->GetFormatTag());
 	}
     else
     {
diff --git a/Engine/Source/Runtime/IOS/IOSAudio/Public/IOSAudioDevice.h b/Engine/Source/Runtime/IOS/IOSAudio/Public/IOSAudioDevice.h
--- a/Engine/Source/Runtime/IOS/IOSAudio/Public/IOSAudioDevice.h
+++ b/Engine/Source/Runtime/IOS/IOSAudio/Public/IOSAudioDevice.h
@@ -12,6 +12,7 @@
 #include "Sound/SoundWave.h"
 #include "DSP/ParamInterpolator.h"
 #include "DSP/OnePole.h"
+#include "AudioDecompress.h"
 
 /*------------------------------------------------------------------------------------
 	Audio Framework system headers
@@ -93,7 +94,7 @@ public:
 	uint32 BufferSize;
 	
 	/** Wrapper to assist in the bookkeeping of uncompressed data when streaming */
-	class FADPCMAudioInfo*		DecompressionState;
+	ICompressedAudioInfo* DecompressionState;
 	bool	bStreaming;
 	bool    bIsProcedural;
 };
