[NO NON-PII COMMIT MESSAGE]

diff --git a/Engine/Source/Runtime/IOS/IOSAudio/Private/IOSAudioSource.cpp b/Engine/Source/Runtime/IOS/IOSAudio/Private/IOSAudioSource.cpp
--- a/Engine/Source/Runtime/IOS/IOSAudio/Private/IOSAudioSource.cpp
+++ b/Engine/Source/Runtime/IOS/IOSAudio/Private/IOSAudioSource.cpp
@@ -336,6 +336,8 @@ void FIOSAudioSoundSource::Stop(void)
 		if(IOSBuffer != NULL)
 		{
 			IOSBuffer->DecompressionState->SeekToTime(0.0f);
+			delete IOSBuffer;
+			IOSBuffer = NULL;
 		}
 	}
 
