2nd attempt to clean up buffer properly on iOS when the source is stopped.

diff --git a/Engine/Source/Runtime/IOS/IOSAudio/Private/IOSAudioSource.cpp b/Engine/Source/Runtime/IOS/IOSAudio/Private/IOSAudioSource.cpp
--- a/Engine/Source/Runtime/IOS/IOSAudio/Private/IOSAudioSource.cpp
+++ b/Engine/Source/Runtime/IOS/IOSAudio/Private/IOSAudioSource.cpp
@@ -104,6 +104,17 @@ FIOSAudioSoundSource::~FIOSAudioSoundSource(void)
 	Buffer = NULL;
 }
 
+void FIOSAudioSoundSource::CleanupAudioBuffer()
+{
+	// Always enure we have a unique FIOSAudioSoundBuffer and that we delete the old one if it exists
+	if (IOSBuffer != NULL)
+	{
+		delete IOSBuffer;
+		IOSBuffer = NULL;
+	}
+}
+
+
 bool FIOSAudioSoundSource::Init(FWaveInstance* InWaveInstance)
 {
 	// Wait for the render callback to finish and then prevent it from being entered again in case this object is deleted after being stopped
@@ -123,11 +134,7 @@ bool FIOSAudioSoundSource::Init(FWaveInstance* InWaveInstance)
 		return false;
 	}
 	
-	// Always enure we have a unique FIOSAudioSoundBuffer and that we delete the old one if it exists
-	if(IOSBuffer != NULL)
-	{
-		delete IOSBuffer;
-	}
+	CleanupAudioBuffer();
 	
 	IOSBuffer = FIOSAudioSoundBuffer::Init(IOSAudioDevice, InWaveInstance->WaveData);
 	Buffer = IOSBuffer;
@@ -330,16 +337,8 @@ void FIOSAudioSoundSource::Stop(void)
 
 	// Call parent class version regardless of if there's a wave instance
 	FSoundSource::Stop();
-		
-	if (WaveInstance)
-	{
-		if(IOSBuffer != NULL)
-		{
-			IOSBuffer->DecompressionState->SeekToTime(0.0f);
-			delete IOSBuffer;
-			IOSBuffer = NULL;
-		}
-	}
+
+	CleanupAudioBuffer();
 
 	// It's now safe to unlock the callback
 	UnlockCallback(&CallbackLock);
diff --git a/Engine/Source/Runtime/IOS/IOSAudio/Public/IOSAudioDevice.h b/Engine/Source/Runtime/IOS/IOSAudio/Public/IOSAudioDevice.h
--- a/Engine/Source/Runtime/IOS/IOSAudio/Public/IOSAudioDevice.h
+++ b/Engine/Source/Runtime/IOS/IOSAudio/Public/IOSAudioDevice.h
@@ -180,6 +180,8 @@ private:
 	                                       const AudioTimeStamp *InTimeStamp, UInt32 InBusNumber,
 	                                       UInt32 InNumberFrames, AudioBufferList *IOData);
 
+	void CleanupAudioBuffer();
+
 	friend class FIOSAudioDevice;
 };
 
