Set minimum iOS version to 11

diff --git a/Engine/Source/Developer/IOS/IOSPlatformEditor/Private/IOSTargetSettingsCustomization.cpp b/Engine/Source/Developer/IOS/IOSPlatformEditor/Private/IOSTargetSettingsCustomization.cpp
--- a/Engine/Source/Developer/IOS/IOSPlatformEditor/Private/IOSTargetSettingsCustomization.cpp
+++ b/Engine/Source/Developer/IOS/IOSPlatformEditor/Private/IOSTargetSettingsCustomization.cpp
@@ -1790,7 +1790,7 @@ void FIOSTargetSettingsCustomization::SetShaderStandard(int32 Value)
 	{
 		FText Message;
 		
-		uint8 EnumValue = (uint8)EIOSVersion::IOS_10;
+		uint8 EnumValue = (uint8)EIOSVersion::IOS_11;
 		if (MinOSPropertyHandle.IsValid())
 		{
 			MinOSPropertyHandle->GetValue(EnumValue);
@@ -1827,9 +1827,9 @@ void FIOSTargetSettingsCustomization::SetShaderStandard(int32 Value)
         }
 
 		// make sure we never set the min version to less than current supported
-		if (((EIOSVersion)EnumValue < EIOSVersion::IOS_10))
+		if (((EIOSVersion)EnumValue < EIOSVersion::IOS_11))
 		{
-			SetMinVersion((int32)EIOSVersion::IOS_10);
+			SetMinVersion((int32)EIOSVersion::IOS_11);
 		}
 
 		
diff --git a/Engine/Source/Programs/UnrealBuildTool/Platform/IOS/UEDeployIOS.cs b/Engine/Source/Programs/UnrealBuildTool/Platform/IOS/UEDeployIOS.cs
--- a/Engine/Source/Programs/UnrealBuildTool/Platform/IOS/UEDeployIOS.cs
+++ b/Engine/Source/Programs/UnrealBuildTool/Platform/IOS/UEDeployIOS.cs
@@ -341,6 +341,7 @@ namespace UnrealBuildTool
 					MinVersion = "9.0";
 					break;
 				case "IOS_10":
+					Log.TraceWarning("IOS 10 is no longer supported in UE4 as of 4.23");
 					MinVersion = "10.0";
 					break;
 				case "IOS_11":
@@ -353,7 +354,7 @@ namespace UnrealBuildTool
 			}
 			else
 			{
-				MinVersion = "9.0";
+				MinVersion = "11.0";
 			}
 
 			// Get Facebook Support details
diff --git a/Engine/Source/Runtime/IOS/IOSRuntimeSettings/Classes/IOSRuntimeSettings.h b/Engine/Source/Runtime/IOS/IOSRuntimeSettings/Classes/IOSRuntimeSettings.h
--- a/Engine/Source/Runtime/IOS/IOSRuntimeSettings/Classes/IOSRuntimeSettings.h
+++ b/Engine/Source/Runtime/IOS/IOSRuntimeSettings/Classes/IOSRuntimeSettings.h
@@ -43,7 +43,7 @@ UENUM()
 	IOS_9 = 9 UMETA(Hidden),
 
 	/** iOS 10 */
-	IOS_10 = 10 UMETA(DisplayName = "10.0"),
+	IOS_10 = 10 UMETA(Hidden),
 
 	/** iOS 11 */
 	IOS_11 = 11 UMETA(DisplayName = "11.0"),
diff --git a/Engine/Source/Runtime/IOS/IOSRuntimeSettings/Private/IOSRuntimeSettings.cpp b/Engine/Source/Runtime/IOS/IOSRuntimeSettings/Private/IOSRuntimeSettings.cpp
--- a/Engine/Source/Runtime/IOS/IOSRuntimeSettings/Private/IOSRuntimeSettings.cpp
+++ b/Engine/Source/Runtime/IOS/IOSRuntimeSettings/Private/IOSRuntimeSettings.cpp
@@ -21,7 +21,7 @@ UIOSRuntimeSettings::UIOSRuntimeSettings(const FObjectInitializer& ObjectInitial
     FrameRateLock = EPowerUsageFrameRateLock::PUFRL_30;
 	bSupportsIPad = true;
 	bSupportsIPhone = true;
-	MinimumiOSVersion = EIOSVersion::IOS_10;
+	MinimumiOSVersion = EIOSVersion::IOS_11;
 	EnableRemoteShaderCompile = false;
 	bGeneratedSYMFile = false;
 	bGeneratedSYMBundle = false;
@@ -154,50 +154,51 @@ void UIOSRuntimeSettings::PostInitProperties()
 		}
 	}
 
-	// switch IOS_6.1, IOS_7, IOS_8, and IOS_9 to IOS_10
-	if (MinimumiOSVersion < EIOSVersion::IOS_10)
+	// switch IOS_6.1, IOS_7, IOS_8, IOS_9 and IOS_10 to IOS_11
+	if (MinimumiOSVersion < EIOSVersion::IOS_11)
 	{
-		MinimumiOSVersion = EIOSVersion::IOS_10;
+		MinimumiOSVersion = EIOSVersion::IOS_11;
+		UpdateSinglePropertyInConfigFile(GetClass()->FindPropertyByName(GET_MEMBER_NAME_CHECKED(UIOSRuntimeSettings, MinimumiOSVersion)), GetDefaultConfigFilename());
 	}
 	if (bSupportsOpenGLES2)
 	{
 		bSupportsOpenGLES2 = false;
+		UpdateSinglePropertyInConfigFile(GetClass()->FindPropertyByName(GET_MEMBER_NAME_CHECKED(UIOSRuntimeSettings, bSupportsOpenGLES2)), GetDefaultConfigFilename());
 	}
 	if (bDevForArmV7)
 	{
 		bDevForArmV7 = false;
+		UpdateSinglePropertyInConfigFile(GetClass()->FindPropertyByName(GET_MEMBER_NAME_CHECKED(UIOSRuntimeSettings, bDevForArmV7)), GetDefaultConfigFilename());
 	}
 	if (bDevForArmV7S)
 	{
 		bDevForArmV7S = false;
+		UpdateSinglePropertyInConfigFile(GetClass()->FindPropertyByName(GET_MEMBER_NAME_CHECKED(UIOSRuntimeSettings, bDevForArmV7S)), GetDefaultConfigFilename());
 	}
 	if (bShipForArmV7)
 	{
 		bShipForArmV7 = false;
+		UpdateSinglePropertyInConfigFile(GetClass()->FindPropertyByName(GET_MEMBER_NAME_CHECKED(UIOSRuntimeSettings, bShipForArmV7)), GetDefaultConfigFilename());
 	}
 	if (bShipForArmV7S)
 	{
 		bShipForArmV7S = false;
+		UpdateSinglePropertyInConfigFile(GetClass()->FindPropertyByName(GET_MEMBER_NAME_CHECKED(UIOSRuntimeSettings, bShipForArmV7S)), GetDefaultConfigFilename());
 	}
 	if (!bSupportsMetal && !bSupportsMetalMRT)
 	{
 		bSupportsMetal = true;
+		UpdateSinglePropertyInConfigFile(GetClass()->FindPropertyByName(GET_MEMBER_NAME_CHECKED(UIOSRuntimeSettings, bSupportsMetal)), GetDefaultConfigFilename());
 	}
 	if (!bDevForArm64)
 	{
 		bDevForArm64 = true;
+		UpdateSinglePropertyInConfigFile(GetClass()->FindPropertyByName(GET_MEMBER_NAME_CHECKED(UIOSRuntimeSettings, bDevForArm64)), GetDefaultConfigFilename());
 	}
 	if (!bShipForArm64)
 	{
 		bShipForArm64 = true;
+		UpdateSinglePropertyInConfigFile(GetClass()->FindPropertyByName(GET_MEMBER_NAME_CHECKED(UIOSRuntimeSettings, bShipForArm64)), GetDefaultConfigFilename());
 	}
-	
-	// Due to a driver bug on A8 devices running iOS 9 we can only support the global clip-plane when running iOS 10+
-	static IConsoleVariable* ClipPlaneCVar = IConsoleManager::Get().FindConsoleVariable(TEXT("r.AllowGlobalClipPlane"));
-	if (ClipPlaneCVar && ClipPlaneCVar->GetInt() != 0 && MinimumiOSVersion < EIOSVersion::IOS_10)
-	{
-		MinimumiOSVersion = EIOSVersion::IOS_10;
-	}
-
 }
 #endif
