[NO NON-PII COMMIT MESSAGE]

diff --git a/Engine/Source/Runtime/IOS/IOSAudio/Private/IOSAudioSource.cpp b/Engine/Source/Runtime/IOS/IOSAudio/Private/IOSAudioSource.cpp
--- a/Engine/Source/Runtime/IOS/IOSAudio/Private/IOSAudioSource.cpp
+++ b/Engine/Source/Runtime/IOS/IOSAudio/Private/IOSAudioSource.cpp
@@ -338,8 +338,6 @@ void FIOSAudioSoundSource::Stop(void)
 	// Call parent class version regardless of if there's a wave instance
 	FSoundSource::Stop();
 
-	CleanupAudioBuffer();
-
 	// It's now safe to unlock the callback
 	UnlockCallback(&CallbackLock);
 }
