Prevent scheduling of local notifications in the past on iOS (check was already present on Android).
#rnx

diff --git a/Engine/Source/Runtime/IOS/IOSLocalNotification/Private/IOSLocalNotification.cpp b/Engine/Source/Runtime/IOS/IOSLocalNotification/Private/IOSLocalNotification.cpp
--- a/Engine/Source/Runtime/IOS/IOSLocalNotification/Private/IOSLocalNotification.cpp
+++ b/Engine/Source/Runtime/IOS/IOSLocalNotification/Private/IOSLocalNotification.cpp
@@ -142,6 +142,11 @@ static int32 NotificationNumber = 0;
 int32 FIOSLocalNotificationService::ScheduleLocalNotificationAtTime(const FDateTime& FireDateTime, bool LocalTime, const FText& Title, const FText& Body, const FText& Action, const FString& ActivationEvent)
 {
 #if !PLATFORM_TVOS
+	if (FireDateTime < (LocalTime ? FDateTime::Now() : FDateTime::UtcNow()))
+	{
+		return -1;
+	}
+
     //Create local copies of these for the block to capture
     FDateTime FireDateTimeCopy = FireDateTime;
     FText TitleCopy = Title;
@@ -173,6 +178,11 @@ int32 FIOSLocalNotificationService::ScheduleLocalNotificationAtTime(const FDateT
 int32 FIOSLocalNotificationService::ScheduleLocalNotificationBadgeAtTime(const FDateTime& FireDateTime, bool LocalTime, const FString& ActivationEvent)
 {
 #if !PLATFORM_TVOS
+	if (FireDateTime < (LocalTime ? FDateTime::Now() : FDateTime::UtcNow()))
+	{
+		return -1;
+	}
+
 	FDateTime FireDateTimeCopy = FireDateTime;
 	FString ActivationEventCopy = ActivationEvent;
 	int32 CurrentNotificationId = NotificationNumber++;
