Merging AudioMixer change from 11.1 to 12 (added logging to see old iOS audio backend initialize)

diff --git a/Engine/Source/Runtime/IOS/IOSAudio/Private/IOSAudioDevice.cpp b/Engine/Source/Runtime/IOS/IOSAudio/Private/IOSAudioDevice.cpp
--- a/Engine/Source/Runtime/IOS/IOSAudio/Private/IOSAudioDevice.cpp
+++ b/Engine/Source/Runtime/IOS/IOSAudio/Private/IOSAudioDevice.cpp
@@ -116,6 +116,7 @@ FAudioPlatformSettings FIOSAudioDevice::GetPlatformSettings() const
 
 bool FIOSAudioDevice::InitializeHardware()
 {
+    UE_LOG(LogIOSAudio, Warning, TEXT("Initializing legacy audio backend for iOS"));
 	SIZE_T SampleSize = sizeof(AudioSampleType);
 	double GraphSampleRate = 44100.0;
 
