Added a project setting to enable files created by UE4 to appear in the iOS on-device Files App (disabled by default)

diff --git a/Engine/Source/Developer/IOS/IOSPlatformEditor/Private/IOSTargetSettingsCustomization.cpp b/Engine/Source/Developer/IOS/IOSPlatformEditor/Private/IOSTargetSettingsCustomization.cpp
--- a/Engine/Source/Developer/IOS/IOSPlatformEditor/Private/IOSTargetSettingsCustomization.cpp
+++ b/Engine/Source/Developer/IOS/IOSPlatformEditor/Private/IOSTargetSettingsCustomization.cpp
@@ -821,6 +821,7 @@ void FIOSTargetSettingsCustomization::BuildPListSection(IDetailLayoutBuilder& De
 	SETUP_PLIST_PROP(PreferredLandscapeOrientation, OrientationCategory);
 	
 	SETUP_PLIST_PROP(bSupportsITunesFileSharing, FileSystemCategory);
+	SETUP_PLIST_PROP(bSupportsFilesApp, FileSystemCategory);
 	
 	SETUP_PLIST_PROP(bSupportsMetal, RenderCategory);
 	
diff --git a/Engine/Source/Programs/UnrealBuildTool/Platform/IOS/UEDeployIOS.cs b/Engine/Source/Programs/UnrealBuildTool/Platform/IOS/UEDeployIOS.cs
--- a/Engine/Source/Programs/UnrealBuildTool/Platform/IOS/UEDeployIOS.cs
+++ b/Engine/Source/Programs/UnrealBuildTool/Platform/IOS/UEDeployIOS.cs
@@ -273,6 +273,8 @@ namespace UnrealBuildTool
 			// ITunes file sharing
 			bool bSupportsITunesFileSharing = false;
 			Ini.GetBool("/Script/IOSRuntimeSettings.IOSRuntimeSettings", "bSupportsITunesFileSharing", out bSupportsITunesFileSharing);
+			bool bSupportsFilesApp = false;
+			Ini.GetBool("/Script/IOSRuntimeSettings.IOSRuntimeSettings", "bSupportsFilesApp", out bSupportsFilesApp);
 
 			// bundle display name
 			string BundleDisplayName;
@@ -451,6 +453,11 @@ namespace UnrealBuildTool
 			Text.AppendLine("\t<true/>");
 			Text.AppendLine("\t<key>UIFileSharingEnabled</key>");
 			Text.AppendLine(string.Format("\t<{0}/>", bSupportsITunesFileSharing ? "true" : "false"));
+			if (bSupportsFilesApp)
+			{
+				Text.AppendLine("\t<key>LSSupportsOpeningDocumentsInPlace</key>");
+				Text.AppendLine("\t<true/>");
+			}
 			Text.AppendLine("\t<key>UIRequiresFullScreen</key>");
 			Text.AppendLine("\t<true/>");
 			Text.AppendLine("\t<key>UIViewControllerBasedStatusBarAppearance</key>");
diff --git a/Engine/Source/Runtime/IOS/IOSRuntimeSettings/Classes/IOSRuntimeSettings.h b/Engine/Source/Runtime/IOS/IOSRuntimeSettings/Classes/IOSRuntimeSettings.h
--- a/Engine/Source/Runtime/IOS/IOSRuntimeSettings/Classes/IOSRuntimeSettings.h
+++ b/Engine/Source/Runtime/IOS/IOSRuntimeSettings/Classes/IOSRuntimeSettings.h
@@ -372,9 +372,14 @@ public:
 	UPROPERTY(GlobalConfig, EditAnywhere, Category = DeviceOrientations)
 	uint32 bSupportsLandscapeRightOrientation : 1;
 
-	UPROPERTY(GlobalConfig, EditAnywhere, Category = FileSystem)
+	// Whether files created by the app will be accessible from the iTunes File Sharing feature
+	UPROPERTY(GlobalConfig, EditAnywhere, Category = FileSystem, meta = (DisplayName = "Support iTunes File Sharing"))
 	uint32 bSupportsITunesFileSharing : 1;
 	
+	// Whether files created by the app will be accessible from within the device's Files app (requires iTunes File Sharing)
+	UPROPERTY(GlobalConfig, EditAnywhere, Category = FileSystem, meta = (DisplayName = "Support Files App", EditCondition = "bSupportsITunesFileSharing"))
+	uint32 bSupportsFilesApp : 1;
+	
 	// The Preferred Orientation will be used as the initial orientation at launch when both Landscape Left and Landscape Right orientations are to be supported.
 	UPROPERTY(GlobalConfig, EditAnywhere, Category = DeviceOrientations, meta = (DisplayName = "Preferred Landscape Orientation"))
 	EIOSLandscapeOrientation PreferredLandscapeOrientation;
diff --git a/Engine/Source/Runtime/IOS/IOSRuntimeSettings/Private/IOSRuntimeSettings.cpp b/Engine/Source/Runtime/IOS/IOSRuntimeSettings/Private/IOSRuntimeSettings.cpp
--- a/Engine/Source/Runtime/IOS/IOSRuntimeSettings/Private/IOSRuntimeSettings.cpp
+++ b/Engine/Source/Runtime/IOS/IOSRuntimeSettings/Private/IOSRuntimeSettings.cpp
@@ -14,6 +14,7 @@ UIOSRuntimeSettings::UIOSRuntimeSettings(const FObjectInitializer& ObjectInitial
 	bEnableCloudKitSupport = false;
 	bSupportsPortraitOrientation = true;
 	bSupportsITunesFileSharing = false;
+	bSupportsFilesApp = false;
 	BundleDisplayName = TEXT("UE4 Game");
 	BundleName = TEXT("MyUE4Game");
 	BundleIdentifier = TEXT("com.YourCompany.GameNameNoSpaces");
