Fix RoboMerge duplicated code causing CIS issue

diff --git a/Engine/Source/Runtime/IOS/IOSRuntimeSettings/Classes/IOSRuntimeSettings.h b/Engine/Source/Runtime/IOS/IOSRuntimeSettings/Classes/IOSRuntimeSettings.h
--- a/Engine/Source/Runtime/IOS/IOSRuntimeSettings/Classes/IOSRuntimeSettings.h
+++ b/Engine/Source/Runtime/IOS/IOSRuntimeSettings/Classes/IOSRuntimeSettings.h
@@ -424,21 +424,13 @@ public:
 	UPROPERTY(GlobalConfig, EditAnywhere, Category = ExtraData)
 	FString AdditionalPlistData;
 	
-	/**
-	 * Chose wether to use a storyboard as a Launchscreen or not. If True, LaunchScreen.storyboard in Build/IOS/Resources/Interface
-	 * in the Engine folder will be copied to the bundle app and the Launchscreen images WILL NOT be included.
-	 * When this setting in changed, delete the Binaries/IOS/Payload folder before recompiling.
-	 */
-	UPROPERTY(GlobalConfig, EditAnywhere, Category = LaunchScreen)
-	bool 	bLaunchscreenStoryboard;
-
 	/**
 	 * Choose whether to use a storyboard as a Launchscreen or not. If True, LaunchScreen.storyboard in Build/IOS/Resources/Interface
 	 * in the Engine folder will be copied to the bundle app and the Launchscreen images WILL NOT be included.
 	 * When using assets in your LaunchScreen.Storyboard, add them in Build/IOS/Resources/Interface/Assets and they will be included.
 	 */
 	UPROPERTY(GlobalConfig, EditAnywhere, Category = LaunchScreen)
-	bool 	bLaunchscreenStoryboard;
+	bool bLaunchscreenStoryboard;
 
 	// Whether the app supports Facebook
 	UPROPERTY(GlobalConfig, EditAnywhere, Category = Online)
