Marked Custom Launchscreen Storyboards experiemental in IOS Project Settings for 4.25.0

diff --git a/Engine/Source/Runtime/IOS/IOSRuntimeSettings/Classes/IOSRuntimeSettings.h b/Engine/Source/Runtime/IOS/IOSRuntimeSettings/Classes/IOSRuntimeSettings.h
--- a/Engine/Source/Runtime/IOS/IOSRuntimeSettings/Classes/IOSRuntimeSettings.h
+++ b/Engine/Source/Runtime/IOS/IOSRuntimeSettings/Classes/IOSRuntimeSettings.h
@@ -430,7 +430,7 @@ public:
 	 * copied to the bundle app and the Launch screen image above will not be included in the app.
 	 * When using assets in your custom LaunchScreen.storyboard, add them in Build/IOS/Resources/Interface/Assets and they will be included.
 	 */
-	UPROPERTY(GlobalConfig, EditAnywhere, Category = LaunchScreen)
+	UPROPERTY(GlobalConfig, EditAnywhere, Category = LaunchScreen, meta = (DisplayName = "Custom Launchscreen Storyboard (experimental)"))
 	bool bCustomLaunchscreenStoryboard;
 
 	// Whether the app supports Facebook
