Fix UE4 reference

diff --git a/Engine/Source/Runtime/IOS/IOSRuntimeSettings/Classes/IOSRuntimeSettings.h b/Engine/Source/Runtime/IOS/IOSRuntimeSettings/Classes/IOSRuntimeSettings.h
--- a/Engine/Source/Runtime/IOS/IOSRuntimeSettings/Classes/IOSRuntimeSettings.h
+++ b/Engine/Source/Runtime/IOS/IOSRuntimeSettings/Classes/IOSRuntimeSettings.h
@@ -272,7 +272,7 @@ public:
 	UPROPERTY(GlobalConfig, EditAnywhere, Category = "Build", meta = (EditCondition = "bUseRSync", ConfigHierarchyEditable))
 	FIOSBuildResourceDirectory DeltaCopyInstallPath;
 
-	// The existing location of an SSH Key found by UE4.
+	// The existing location of an SSH Key found by Unreal Engine.
 	UPROPERTY(VisibleAnywhere, Category = "Build", meta = (DisplayName = "Found Existing SSH permissions file"))
 	FString SSHPrivateKeyLocation;
 
