PR #7737: Fix SSHKey lookup issues of RemoteAddr:Port (Contributed by [REDACTED_PERSON])

diff --git a/Engine/Source/Runtime/IOS/IOSRuntimeSettings/Private/IOSRuntimeSettings.cpp b/Engine/Source/Runtime/IOS/IOSRuntimeSettings/Private/IOSRuntimeSettings.cpp
--- a/Engine/Source/Runtime/IOS/IOSRuntimeSettings/Private/IOSRuntimeSettings.cpp
+++ b/Engine/Source/Runtime/IOS/IOSRuntimeSettings/Private/IOSRuntimeSettings.cpp
@@ -89,8 +89,14 @@ void UIOSRuntimeSettings::PostInitProperties()
 	{
 		SSHPrivateKeyLocation = TEXT("");
 
+		FString RealRemoteServerName = RemoteServerName;
+		if(RemoteServerName.Contains(TEXT(":")))
+		{
+			FString RemoteServerPort;
+			RemoteServerName.Split(TEXT(":"),&RealRemoteServerName,&RemoteServerPort);
+		}
 		const FString DefaultKeyFilename = TEXT("RemoteToolChainPrivate.key");
-		const FString RelativeFilePathLocation = FPaths::Combine(TEXT("SSHKeys"), *RemoteServerName, *RSyncUsername, *DefaultKeyFilename);
+		const FString RelativeFilePathLocation = FPaths::Combine(TEXT("SSHKeys"), *RealRemoteServerName, *RSyncUsername, *DefaultKeyFilename);
 
 		FString Path = FPlatformMisc::GetEnvironmentVariable(TEXT("APPDATA"));
 
