Fix uninitialized variable warning

diff --git a/Engine/Source/Runtime/IOS/LaunchDaemonMessages/Classes/IOSMessageProtocol.h b/Engine/Source/Runtime/IOS/LaunchDaemonMessages/Classes/IOSMessageProtocol.h
--- a/Engine/Source/Runtime/IOS/LaunchDaemonMessages/Classes/IOSMessageProtocol.h
+++ b/Engine/Source/Runtime/IOS/LaunchDaemonMessages/Classes/IOSMessageProtocol.h
@@ -49,6 +49,7 @@ struct FIOSLaunchDaemonPong
 	FIOSLaunchDaemonPong()
 		: bCanPowerOff(false)
 		, bCanPowerOn(false)
+		, bIsAuthorized(false)
 		, bCanReboot(false)
 	{}
 
