Fix IOSAudioBuffer.cpp compile error

diff --git a/Engine/Source/Runtime/IOS/IOSAudio/Private/IOSAudioBuffer.cpp b/Engine/Source/Runtime/IOS/IOSAudio/Private/IOSAudioBuffer.cpp
--- a/Engine/Source/Runtime/IOS/IOSAudio/Private/IOSAudioBuffer.cpp
+++ b/Engine/Source/Runtime/IOS/IOSAudio/Private/IOSAudioBuffer.cpp
@@ -134,7 +134,7 @@ bool FIOSAudioSoundBuffer::ReadCompressedInfo(USoundWave* InWave)
 		return false;
 	}
 
-	return DecompressionState->ReadCompressedInfo(InWave->GetResourceData(), InWave->SoundWaveDataPtr->GetResourceSize(), &QualityInfo);
+	return DecompressionState->ReadCompressedInfo(InWave->GetResourceData(), InWave->GetResourceSize(), &QualityInfo);
 }
 
 bool FIOSAudioSoundBuffer::ReadCompressedData( uint8* Destination, int32 NumFramesToDecode, bool bLooping )
