CIS accessing member that is now private (replacing w/ appropriate getter)

diff --git a/Engine/Source/Runtime/IOS/IOSAudio/Private/IOSAudioBuffer.cpp b/Engine/Source/Runtime/IOS/IOSAudio/Private/IOSAudioBuffer.cpp
--- a/Engine/Source/Runtime/IOS/IOSAudio/Private/IOSAudioBuffer.cpp
+++ b/Engine/Source/Runtime/IOS/IOSAudio/Private/IOSAudioBuffer.cpp
@@ -128,7 +128,7 @@ bool FIOSAudioSoundBuffer::ReadCompressedInfo(USoundWave* InWave)
 	}
 
 	InWave->InitAudioResource(FName(TEXT("ADPCM")));
-	if (InWave->SoundWaveDataPtr->ResourceSize <= 0)
+	if (InWave->GetResourceSize() <= 0)
 	{
 		InWave->RemoveAudioResource();
 		return false;
