Clarify iOS dSYM generation checkbox descriptions

diff --git a/Engine/Source/Runtime/IOS/IOSRuntimeSettings/Classes/IOSRuntimeSettings.h b/Engine/Source/Runtime/IOS/IOSRuntimeSettings/Classes/IOSRuntimeSettings.h
--- a/Engine/Source/Runtime/IOS/IOSRuntimeSettings/Classes/IOSRuntimeSettings.h
+++ b/Engine/Source/Runtime/IOS/IOSRuntimeSettings/Classes/IOSRuntimeSettings.h
@@ -229,11 +229,11 @@ public:
 	FIOSBuildResourceDirectory WindowsMetalToolchainOverride;
 
 	// Enable generation of dSYM file
-	UPROPERTY(GlobalConfig, EditAnywhere, Category = Build, meta = (DisplayName = "Generate dSYM file for code debugging and profiling"))
+	UPROPERTY(GlobalConfig, EditAnywhere, Category = Build, meta = (DisplayName = "Generate dSYMs for code debugging and profiling"))
 	bool bGeneratedSYMFile;
 
 	// Enable generation of dSYM bundle
-	UPROPERTY(GlobalConfig, EditAnywhere, Category = Build, meta = (DisplayName = "Generate dSYM bundle for third party crash tools"))
+	UPROPERTY(GlobalConfig, EditAnywhere, Category = Build, meta = (DisplayName = "Generate dSYMs as a bundle for third party crash tools"), meta = (EditCondition = "bGeneratedSYMFile"))
 	bool bGeneratedSYMBundle;
 
 	// Enable generation of a .udebugsymbols file, which allows offline, platform-independent symbolication for the Malloc Profiler or external crash reporting tools. Requires a dSYM file or bundle.
