fixing compile error

- Fixing up legacy backend (that should have been deleted) to implement the changed interface.

diff --git a/Engine/Source/Runtime/IOS/IOSAudio/Private/IOSAudioDevice.cpp b/Engine/Source/Runtime/IOS/IOSAudio/Private/IOSAudioDevice.cpp
--- a/Engine/Source/Runtime/IOS/IOSAudio/Private/IOSAudioDevice.cpp
+++ b/Engine/Source/Runtime/IOS/IOSAudio/Private/IOSAudioDevice.cpp
@@ -292,7 +292,7 @@ void FIOSAudioDevice::HandleError(const TCHAR* InLogOutput, bool bTeardown)
 	}
 }
 
-class ICompressedAudioInfo* FIOSAudioDevice::CreateCompressedAudioInfo(USoundWave* SoundWave)
+class ICompressedAudioInfo* FIOSAudioDevice::CreateCompressedAudioInfo(const FName& InRuntimeFormat) const
 {
 	return new FADPCMAudioInfo();
 }
diff --git a/Engine/Source/Runtime/IOS/IOSAudio/Public/IOSAudioDevice.h b/Engine/Source/Runtime/IOS/IOSAudio/Public/IOSAudioDevice.h
--- a/Engine/Source/Runtime/IOS/IOSAudio/Public/IOSAudioDevice.h
+++ b/Engine/Source/Runtime/IOS/IOSAudio/Public/IOSAudioDevice.h
@@ -195,7 +195,7 @@ public:
 	FIOSAudioDevice();
 	virtual ~FIOSAudioDevice() { }
 	
-	virtual FName GetRuntimeFormat(USoundWave* SoundWave) override
+	virtual FName GetRuntimeFormat(const USoundWave* SoundWave) const override
 	{
 		static FName NAME_ADPCM(TEXT("ADPCM"));
 		return NAME_ADPCM;
@@ -251,7 +251,7 @@ private:
 
 	virtual bool SupportsRealtimeDecompression() const override { return true; }
 	
-	virtual class ICompressedAudioInfo* CreateCompressedAudioInfo(USoundWave* SoundWave) override;
+	virtual ICompressedAudioInfo* CreateCompressedAudioInfo(const FName& InRuntimeFormat) const override;
 
 	void HandleError(const TCHAR* InLogOutput, bool bTeardown = false);
 
