Fix objc retain leaks in IOSLocalNotifications (found via Static Analysis)

#rnx

diff --git a/Engine/Source/Runtime/IOS/IOSLocalNotification/Private/IOSLocalNotification.cpp b/Engine/Source/Runtime/IOS/IOSLocalNotification/Private/IOSLocalNotification.cpp
--- a/Engine/Source/Runtime/IOS/IOSLocalNotification/Private/IOSLocalNotification.cpp
+++ b/Engine/Source/Runtime/IOS/IOSLocalNotification/Private/IOSLocalNotification.cpp
@@ -96,23 +96,24 @@ public:
 		{
 			// if not local time, convert from UTC to local time, UNCalendarNotificationTrigger misbehaves
 			// if the components have anything more than previously set (can't specify timeZone)
-            NSCalendar *utcCalendar = [NSCalendar calendarWithIdentifier:NSCalendarIdentifierISO8601];
-            utcCalendar.timeZone = [NSTimeZone timeZoneForSecondsFromGMT:0];
-            NSDate *utcDate = [utcCalendar dateFromComponents:dateComps];
-            NSDateComponents *utcDateComps = [utcCalendar componentsInTimeZone:calendar.timeZone fromDate:utcDate];
-
+			NSCalendar *utcCalendar = [NSCalendar calendarWithIdentifier:NSCalendarIdentifierISO8601];
+			utcCalendar.timeZone = [NSTimeZone timeZoneForSecondsFromGMT:0];
+			NSDate *utcDate = [utcCalendar dateFromComponents:dateComps];
+			NSDateComponents *utcDateComps = [utcCalendar componentsInTimeZone:calendar.timeZone fromDate:utcDate];
+			
 			// now copy components back because utcDateComps has too many fields set for
 			// UNCalendarNotificationTrigger and will now work properly on all devices
-            dateComps.day = utcDateComps.day;
-            dateComps.month = utcDateComps.month;
-            dateComps.year = utcDateComps.year;
-            dateComps.hour = utcDateComps.hour;
-            dateComps.minute = utcDateComps.minute;
-            dateComps.second = utcDateComps.second;
+			dateComps.day = utcDateComps.day;
+			dateComps.month = utcDateComps.month;
+			dateComps.year = utcDateComps.year;
+			dateComps.hour = utcDateComps.hour;
+			dateComps.minute = utcDateComps.minute;
+			dateComps.second = utcDateComps.second;
 		}
 		
 		UNCalendarNotificationTrigger *trigger = [UNCalendarNotificationTrigger triggerWithDateMatchingComponents:dateComps repeats:NO];
 		
+		[dateComps release];
 		return trigger;
 	}
 #endif
@@ -147,28 +148,29 @@ int32 FIOSLocalNotificationService::ScheduleLocalNotificationAtTime(const FDateT
 		return -1;
 	}
 
-    //Create local copies of these for the block to capture
-    FDateTime FireDateTimeCopy = FireDateTime;
-    FText TitleCopy = Title;
-    FText BodyCopy = Body;
-    FText ActionCopy = Action;
-    FString ActivationEventCopy = ActivationEvent;
+	//Create local copies of these for the block to capture
+	FDateTime FireDateTimeCopy = FireDateTime;
+	FText TitleCopy = Title;
+	FText BodyCopy = Body;
+	FText ActionCopy = Action;
+	FString ActivationEventCopy = ActivationEvent;
 	int32 CurrentNotificationId = NotificationNumber++;
 	
 	//have to schedule notification on main thread queue
 	dispatch_async(dispatch_get_main_queue(), ^{
-        UNMutableNotificationContent* Content = FIOSLocalNotificationModule::CreateNotificationContent(TitleCopy, BodyCopy, ActionCopy, ActivationEventCopy, 1);
-        UNCalendarNotificationTrigger* Trigger = FIOSLocalNotificationModule::CreateCalendarNotificationTrigger(FireDateTimeCopy, LocalTime);
-        
-        UNNotificationRequest *Request = [UNNotificationRequest requestWithIdentifier:@(CurrentNotificationId).stringValue content:Content trigger:Trigger];
-
+		UNMutableNotificationContent* Content = FIOSLocalNotificationModule::CreateNotificationContent(TitleCopy, BodyCopy, ActionCopy, ActivationEventCopy, 1);
+		UNCalendarNotificationTrigger* Trigger = FIOSLocalNotificationModule::CreateCalendarNotificationTrigger(FireDateTimeCopy, LocalTime);
+		
+		UNNotificationRequest *Request = [UNNotificationRequest requestWithIdentifier:@(CurrentNotificationId).stringValue content:Content trigger:Trigger];
+		
 		UNUserNotificationCenter *Center = [UNUserNotificationCenter currentNotificationCenter];
-        [Center addNotificationRequest : Request withCompletionHandler : ^ (NSError * _Nullable error) {
-            if (error != nil) {
-                UE_LOG(LogIOSLocalNotification, Warning, TEXT("Error scheduling notification: %d"), CurrentNotificationId);
-            }
-        }];
-    });
+		[Center addNotificationRequest : Request withCompletionHandler : ^ (NSError * _Nullable error) {
+			if (error != nil) {
+				UE_LOG(LogIOSLocalNotification, Warning, TEXT("Error scheduling notification: %d"), CurrentNotificationId);
+			}
+		}];
+		[Content release];
+	});
 	return CurrentNotificationId;
 #else
 	return -1;
@@ -200,6 +202,7 @@ int32 FIOSLocalNotificationService::ScheduleLocalNotificationBadgeAtTime(const F
 				UE_LOG(LogIOSLocalNotification, Warning, TEXT("Error scheduling notification: %d"), CurrentNotificationId);
 			}
 		}];
+		[Content release];
 	});
 	
 	return CurrentNotificationId;
