Clean up project settings for Secondary Mac featre

diff --git a/Engine/Source/Runtime/IOS/IOSRuntimeSettings/Classes/IOSRuntimeSettings.h b/Engine/Source/Runtime/IOS/IOSRuntimeSettings/Classes/IOSRuntimeSettings.h
--- a/Engine/Source/Runtime/IOS/IOSRuntimeSettings/Classes/IOSRuntimeSettings.h
+++ b/Engine/Source/Runtime/IOS/IOSRuntimeSettings/Classes/IOSRuntimeSettings.h
@@ -205,6 +205,10 @@ public:
 	UPROPERTY(GlobalConfig, EditAnywhere, Category = Rendering, meta = (DisplayName = "Metal Desktop Renderer"))
 	bool bSupportsMetalMRT;
 	
+	// Support a secondary remote mac ?
+	UPROPERTY(GlobalConfig, EditAnywhere, Category = "Build", meta = (DisplayName = "Supports a secondary remote mac to facilitate debug of iOS / tvOS builds"))
+	bool bSupportSecondaryMac;
+
     // Whether to build the iOS project as a framework.
     UPROPERTY(GlobalConfig, EditAnywhere, Category = Build, meta = (DisplayName = "Build project as a framework (Experimental)"))
     bool bBuildAsFramework;
@@ -248,13 +252,8 @@ public:
 	UPROPERTY(GlobalConfig, EditAnywhere, Category = "Build", meta = (ConfigHierarchyEditable))
 	FString RemoteServerName;
 
-	// The name or ip address of the remote mac which will be used to build IOS
-	UPROPERTY(GlobalConfig, EditAnywhere, Category = "Build", meta = (ConfigHierarchyEditable))
-	FString SecondaryRemoteServerName;
-
-	// Support a secondary remote mac ?
-	UPROPERTY(GlobalConfig, EditAnywhere, Category = "Build", meta = (DisplayName = "Supports a secondary remote mac to facilitate debug of iOS / tvOS builds"))
-	bool bSupportSecondaryMac;
+	
+	
 
 	// Enable the use of RSync for remote builds on a mac
 	UPROPERTY(GlobalConfig, EditAnywhere, Category = "Build", meta = (DisplayName = "Use RSync for building IOS", ConfigHierarchyEditable))
@@ -264,44 +263,51 @@ public:
 	UPROPERTY(GlobalConfig, EditAnywhere, Category = "Build", meta = (EditCondition = "bUseRSync", DisplayName = "Username on Remote Server.", ConfigHierarchyEditable))
 	FString RSyncUsername;
 
-	// The secondary mac users name which matches the SSH Private Key, for remote builds using RSync.
-	UPROPERTY(GlobalConfig, EditAnywhere, Category = "Build", meta = (EditCondition = "bUseRSync", DisplayName = "Username on Remote Server.", ConfigHierarchyEditable))
-	FString SecondaryRSyncUsername;
+
 
 	// Optional path on the remote mac where the build files will be copied. If blank, ~/UE5/Builds will be used.
 	UPROPERTY(GlobalConfig, EditAnywhere, Category = "Build", meta = (ConfigHierarchyEditable))
 	FString RemoteServerOverrideBuildPath;
 
-	// Optional path on the secondary remote mac where the build files will be copied. If blank, ~/UE5/Builds will be used.
-	UPROPERTY(GlobalConfig, EditAnywhere, Category = "Build", meta = (ConfigHierarchyEditable))
-	FString SecondaryRemoteServerOverrideBuildPath;
-
+	
 	// The install directory of cwrsync.
 	UPROPERTY(GlobalConfig, EditAnywhere, Category = "Build", meta = (EditCondition = "bUseRSync", ConfigHierarchyEditable))
 	FIOSBuildResourceDirectory CwRsyncInstallPath;
 
-	// The install directory of cwrsync.
-	UPROPERTY(GlobalConfig, EditAnywhere, Category = "Build", meta = (EditCondition = "bUseRSync", ConfigHierarchyEditable))
-	FIOSBuildResourceDirectory SecondaryCwRsyncInstallPath;
-
 	// The existing location of an SSH Key found by Unreal Engine.
 	UPROPERTY(VisibleAnywhere, Category = "Build", meta = (DisplayName = "Found Existing SSH permissions file"))
 	FString SSHPrivateKeyLocation;
 
-	// The existing location of an SSH Key found by Unreal Engine.
-	UPROPERTY(VisibleAnywhere, Category = "Build", meta = (DisplayName = "Found Existing SSH permissions file for Secondary Mac"))
-	FString SecondarySSHPrivateKeyLocation;
-
 	// The path of the ssh permissions key to be used when connecting to the remote server.
 	UPROPERTY(GlobalConfig, EditAnywhere, Category = "Build", meta = (EditCondition = "bUseRSync", DisplayName = "Override existing SSH permissions file", ConfigHierarchyEditable))
 	FIOSBuildResourceFilePath SSHPrivateKeyOverridePath;
 
+	// The name or ip address of the remote mac which will be used to build IOS
+	UPROPERTY(GlobalConfig, EditAnywhere, Category = "Build", meta = (EditCondition = "bSupportSecondaryMac", ConfigHierarchyEditable))
+	FString SecondaryRemoteServerName;
+
+	// Optional path on the secondary remote mac where the build files will be copied. If blank, ~/UE5/Builds will be used.
+	UPROPERTY(GlobalConfig, EditAnywhere, Category = "Build", meta = (EditCondition = "bSupportSecondaryMac", ConfigHierarchyEditable))
+	FString SecondaryRemoteServerOverrideBuildPath;
+
+	// The existing location of an SSH Key found by Unreal Engine.
+	UPROPERTY(VisibleAnywhere, Category = "Build", meta = (EditCondition = "bSupportSecondaryMac", DisplayName = "Found Existing SSH permissions file for Secondary Mac"))
+	FString SecondarySSHPrivateKeyLocation;
+
 	// The path of the ssh permissions key to be used when connecting to the remote server.
-	UPROPERTY(GlobalConfig, EditAnywhere, Category = "Build", meta = (EditCondition = "bUseRSync", DisplayName = "Override existing SSH permissions file for Secondary Mac", ConfigHierarchyEditable))
+	UPROPERTY(GlobalConfig, EditAnywhere, Category = "Build", meta = (EditCondition = "bSupportSecondaryMac", DisplayName = "Override existing SSH permissions file for Secondary Mac", ConfigHierarchyEditable))
 	FIOSBuildResourceFilePath SecondarySSHPrivateKeyOverridePath;
-    
-    // Should the app be compatible with Multi-User feature on tvOS ?　If checked, the game will will shutdown with the typical exit flow.
-    UPROPERTY(GlobalConfig, EditAnywhere, Category = "Build", meta = (DisplayName = "Support user switching on tvOS."))
+
+	// The install directory of cwrsync.
+	UPROPERTY(GlobalConfig, EditAnywhere, Category = "Build", meta = (EditCondition = "bSupportSecondaryMac", ConfigHierarchyEditable))
+	FIOSBuildResourceDirectory SecondaryCwRsyncInstallPath;
+
+	// The secondary mac users name which matches the SSH Private Key, for remote builds using RSync.
+	UPROPERTY(GlobalConfig, EditAnywhere, Category = "Build", meta = (EditCondition = "bSupportSecondaryMac", DisplayName = "Username on Remote Server.", ConfigHierarchyEditable))
+	FString SecondaryRSyncUsername;
+
+	// Should the app be compatible with Multi-User feature on tvOS ?　If checked, the game will will shutdown with the typical exit flow.
+    UPROPERTY(GlobalConfig, EditAnywhere, Category = "Build", meta = (DisplayName = "Support user switching on tvOS"))
     bool bRunAsCurrentUser;
 
 	// If checked, the game will be able to handle multiple gamepads at the same time (the Siri Remote is a gamepad)
