PR #7614: Added missing local push notifications functionalities. (Contributed by [REDACTED_PERSON])

diff --git a/Engine/Source/Runtime/IOS/IOSLocalNotification/Private/IOSLocalNotification.cpp b/Engine/Source/Runtime/IOS/IOSLocalNotification/Private/IOSLocalNotification.cpp
--- a/Engine/Source/Runtime/IOS/IOSLocalNotification/Private/IOSLocalNotification.cpp
+++ b/Engine/Source/Runtime/IOS/IOSLocalNotification/Private/IOSLocalNotification.cpp
@@ -133,9 +133,10 @@ FIOSLocalNotificationService::FIOSLocalNotificationService()
 void FIOSLocalNotificationService::ClearAllLocalNotifications()
 {
 #if !PLATFORM_TVOS
-	UNUserNotificationCenter *Center = [UNUserNotificationCenter currentNotificationCenter];
-	
-	[Center removeAllPendingNotificationRequests];
+	dispatch_async(dispatch_get_main_queue(), ^{
+		UNUserNotificationCenter *Center = [UNUserNotificationCenter currentNotificationCenter];
+		[Center removeAllPendingNotificationRequests];
+	});
 #endif
 }
 
