Support high refresh rates on iPhone

diff --git a/Engine/Config/BaseEngine.ini b/Engine/Config/BaseEngine.ini
--- a/Engine/Config/BaseEngine.ini
+++ b/Engine/Config/BaseEngine.ini
@@ -2836,6 +2836,7 @@ bSupportsLandscapeRightOrientation=True
 PreferredLandscapeOrientation=LandscapeLeft
 bSupportsMetal=True
 bSupportsMetalMRT=False
+bSupportHighRefreshRates=False
 bGameSupportsMultipleActiveControllers=False
 bAllowRemoteRotation=True
 bUseAbsoluteDpadValues=False
diff --git a/Engine/Source/Programs/UnrealBuildTool/Platform/IOS/IOSExports.cs b/Engine/Source/Programs/UnrealBuildTool/Platform/IOS/IOSExports.cs
--- a/Engine/Source/Programs/UnrealBuildTool/Platform/IOS/IOSExports.cs
+++ b/Engine/Source/Programs/UnrealBuildTool/Platform/IOS/IOSExports.cs
@@ -528,6 +528,16 @@ namespace UnrealBuildTool
 					Text.AppendLine("\t<array><string>runs-as-current-user</string></array>");
 				}
 
+				// Support high refresh rates (iPhone only)
+				// https://developer.apple.com/documentation/quartzcore/optimizing_promotion_refresh_rates_for_iphone_13_pro_and_ipad_pro
+				bool bSupportHighRefreshRates = false;
+				PlatformGameConfig.GetBool("/Script/IOSRuntimeSettings.IOSRuntimeSettings", "bSupportHighRefreshRates", out bSupportHighRefreshRates);
+
+				if (bSupportHighRefreshRates && Platform == UnrealTargetPlatform.IOS)
+				{
+					Text.AppendLine("\t<key>CADisableMinimumFrameDurationOnPhone</key><true/>");
+				}
+
 
 				// End of entitlements
 				Text.AppendLine("</dict>");
diff --git a/Engine/Source/Runtime/IOS/IOSRuntimeSettings/Classes/IOSRuntimeSettings.h b/Engine/Source/Runtime/IOS/IOSRuntimeSettings/Classes/IOSRuntimeSettings.h
--- a/Engine/Source/Runtime/IOS/IOSRuntimeSettings/Classes/IOSRuntimeSettings.h
+++ b/Engine/Source/Runtime/IOS/IOSRuntimeSettings/Classes/IOSRuntimeSettings.h
@@ -204,6 +204,10 @@ public:
 	// Whether or not to compile iOS Metal shaders for the desktop renderer (requires iOS 10+ and an A10 processor)
 	UPROPERTY(GlobalConfig, EditAnywhere, Category = Rendering, meta = (DisplayName = "Metal Desktop Renderer"))
 	bool bSupportsMetalMRT;
+    
+    // Should the app be compatible for high refresh rate (iPhone only)
+    UPROPERTY(GlobalConfig, EditAnywhere, Category = "Rendering", meta = (DisplayName = "Enable ProMotion 120Hz on supported iPhone devices"))
+    bool bSupportHighRefreshRates;
 	
 	// Support a secondary remote mac ?
 	UPROPERTY(GlobalConfig, EditAnywhere, Category = "Build", meta = (DisplayName = "Supports a secondary remote mac to facilitate debug of iOS / tvOS builds"))
diff --git a/Engine/Source/Runtime/IOS/IOSRuntimeSettings/Private/IOSRuntimeSettings.cpp b/Engine/Source/Runtime/IOS/IOSRuntimeSettings/Private/IOSRuntimeSettings.cpp
--- a/Engine/Source/Runtime/IOS/IOSRuntimeSettings/Private/IOSRuntimeSettings.cpp
+++ b/Engine/Source/Runtime/IOS/IOSRuntimeSettings/Private/IOSRuntimeSettings.cpp
@@ -49,6 +49,7 @@ UIOSRuntimeSettings::UIOSRuntimeSettings(const FObjectInitializer& ObjectInitial
     bEnableBackgroundFetch = false;
 	bSupportsMetal = true;
 	bSupportsMetalMRT = false;
+    bSupportHighRefreshRates = false;
 	bDisableHTTPS = false;
     bSupportsBackgroundAudio = false;
 }
