iOS Project Settings cleanup

diff --git a/Engine/Config/BaseEngine.ini b/Engine/Config/BaseEngine.ini
--- a/Engine/Config/BaseEngine.ini
+++ b/Engine/Config/BaseEngine.ini
@@ -2870,7 +2870,6 @@ SSHPrivateKeyOverridePath=
 SecondarySSHPrivateKeyOverridePath=
 bEnableRemoteNotificationsSupport=False
 bEnableCloudKitSupport=False
-bRunAsCurrentUser=False
 bSupportSecondaryMac=False
 IOSCloudKitSyncStrategy=None
 bGenerateCrashReportSymbols=false
@@ -3030,7 +3029,7 @@ bBuildAllSupportedOnBuildMachine=true
 bEnableRayTracing=false
 
 [/Script/MacTargetPlatform.XcodeProjectSettings]
-bUseModernXcode=false
+bUseModernXcode=true
 bUseModernCodeSigning=true
 ModernSigningTeam=
 ModernSigningPrefix=com.YourCompany
@@ -3323,4 +3322,4 @@ MaterialParents=(("MI_Default_Opaque", "/Interchange/gltf/MaterialInstances/MI_D
 
 [BuildPatchServices]
 ; Enable/Disable mechanism that provides a disk overflow store for patch chunks.
-bEnableDiskOverflowStore=true
\ No newline at end of file
+bEnableDiskOverflowStore=true
diff --git a/Engine/Source/Developer/IOS/IOSPlatformEditor/Private/IOSTargetSettingsCustomization.cpp b/Engine/Source/Developer/IOS/IOSPlatformEditor/Private/IOSTargetSettingsCustomization.cpp
--- a/Engine/Source/Developer/IOS/IOSPlatformEditor/Private/IOSTargetSettingsCustomization.cpp
+++ b/Engine/Source/Developer/IOS/IOSPlatformEditor/Private/IOSTargetSettingsCustomization.cpp
@@ -316,11 +316,8 @@ void FIOSTargetSettingsCustomization::BuildPListSection(IDetailLayoutBuilder& De
 	IDetailCategoryBuilder& OrientationCategory = DetailLayout.EditCategory(TEXT("Orientation"));
 	IDetailCategoryBuilder& FileSystemCategory = DetailLayout.EditCategory(TEXT("FileSystem"));
 	IDetailCategoryBuilder& RenderCategory = DetailLayout.EditCategory(TEXT("Rendering"));
-	IDetailCategoryBuilder& OSInfoCategory = DetailLayout.EditCategory(TEXT("OS Info"));
-	IDetailCategoryBuilder& DeviceCategory = DetailLayout.EditCategory(TEXT("Devices"));
 	IDetailCategoryBuilder& BuildCategory = DetailLayout.EditCategory(TEXT("Build"));
 	IDetailCategoryBuilder& OnlineCategory = DetailLayout.EditCategory(TEXT("Online"));
-	IDetailCategoryBuilder& ExtraCategory = DetailLayout.EditCategory(TEXT("Extra PList Data"));
 	MobileProvisionProperty = DetailLayout.GetProperty(GET_MEMBER_NAME_CHECKED(UIOSRuntimeSettings, MobileProvision));
 	BuildCategory.AddProperty(MobileProvisionProperty)
 		.Visibility(EVisibility::Hidden);
@@ -717,27 +714,6 @@ void FIOSTargetSettingsCustomization::BuildPListSection(IDetailLayoutBuilder& De
 				]
 			]
 		];
-	
-	BundleCategory.AddCustomRow(LOCTEXT("UpgradeInfo", "Upgrade Info"), false)
-	.WholeRowWidget
-	[
-		SNew(SBorder)
-		.Padding(1)
-		[
-			SNew(SHorizontalBox)
-			+ SHorizontalBox::Slot()
-			.Padding(FMargin(10, 10, 10, 10))
-			.FillWidth(1.0f)
-			[
-				SNew(SRichTextBlock)
-				.Text(LOCTEXT("IOSUpgradeInfoMessage", "<RichTextBlock.TextHighlight>Note to users from 4.6 or earlier</>: We now <RichTextBlock.TextHighlight>GENERATE</> an Info.plist when building, so if you have customized your .plist file, you will need to put all of your changes into the below settings. Note that we don't touch the .plist file that is in your project directory, so you can use it as reference."))
-				.TextStyle(FAppStyle::Get(), "MessageLog")
-				.DecoratorStyleSet(&FAppStyle::Get())
-				.AutoWrapText(true)
-				// + SRichTextBlock::HyperlinkDecorator(TEXT("browser"), FSlateHyperlinkRun::FOnClick::CreateStatic(&OnBrowserLinkClicked))
-			 ]
-		 ]
-	 ];
 
 	// Show properties that are gated by the plist being present and writable
 	RunningIPPProcess = false;
@@ -803,10 +779,6 @@ void FIOSTargetSettingsCustomization::BuildPListSection(IDetailLayoutBuilder& De
 	FSimpleDelegate OnUpdateOSVersionWarning = FSimpleDelegate::CreateSP(this, &FIOSTargetSettingsCustomization::UpdateOSVersionWarning);
 	FSimpleDelegate OnEnableMetalMRT = FSimpleDelegate::CreateSP(this, &FIOSTargetSettingsCustomization::UpdateMetalMRTWarning);
 
-	/* MinOSPropertyHandle = DetailLayout.GetProperty(GET_MEMBER_NAME_CHECKED(UIOSRuntimeSettings, MinimumiOSVersion));
-	MinOSPropertyHandle->SetOnPropertyValueChanged(OnUpdateShaderStandardWarning);
-	OSInfoCategory.AddProperty(MinOSPropertyHandle);*/
-
 	SETUP_PLIST_PROP(BundleDisplayName, BundleCategory);
 	SETUP_PLIST_PROP(BundleName, BundleCategory);
 	SETUP_STATUS_PROP(BundleIdentifier, BundleCategory);
@@ -878,7 +850,7 @@ void FIOSTargetSettingsCustomization::BuildPListSection(IDetailLayoutBuilder& De
 		MinOSPropertyHandle->SetOnPropertyValueChanged(OnUpdateOSVersionWarning);
 
 		// Drop-downs for setting type of lower and upper bound normalization
-		IDetailPropertyRow& MinOSPropertyRow = OSInfoCategory.AddProperty(MinOSPropertyHandle.ToSharedRef());
+		IDetailPropertyRow& MinOSPropertyRow = BuildCategory.AddProperty(MinOSPropertyHandle.ToSharedRef());
 		MinOSPropertyRow.CustomWidget()
 		.NameContent()
 		[
@@ -914,11 +886,6 @@ void FIOSTargetSettingsCustomization::BuildPListSection(IDetailLayoutBuilder& De
 
 		UpdateOSVersionWarning();
 	}
-
-	SETUP_PLIST_PROP(bSupportsIPad, DeviceCategory);
-	SETUP_PLIST_PROP(bSupportsIPhone, DeviceCategory);
-	SETUP_PLIST_PROP(AdditionalPlistData, ExtraCategory);
-
 #undef SETUP_SOURCEONLY_PROP
 }
 
diff --git a/Engine/Source/Programs/UnrealBuildTool/Platform/IOS/IOSExports.cs b/Engine/Source/Programs/UnrealBuildTool/Platform/IOS/IOSExports.cs
--- a/Engine/Source/Programs/UnrealBuildTool/Platform/IOS/IOSExports.cs
+++ b/Engine/Source/Programs/UnrealBuildTool/Platform/IOS/IOSExports.cs
@@ -519,10 +519,7 @@ namespace UnrealBuildTool
 				}
 
 				// Add Multi-user support for tvOS
-				bool bRunAsCurrentUser = false;
-				PlatformGameConfig.GetBool("/Script/IOSRuntimeSettings.IOSRuntimeSettings", "bRunAsCurrentUser", out bRunAsCurrentUser);
-
-				if (bRunAsCurrentUser && Platform == UnrealTargetPlatform.TVOS)
+				if (Platform == UnrealTargetPlatform.TVOS)
 				{
 					Text.AppendLine("\t<key>com.apple.developer.user-management</key>");
 					Text.AppendLine("\t<array><string>runs-as-current-user</string></array>");
diff --git a/Engine/Source/Runtime/IOS/IOSRuntimeSettings/Classes/IOSRuntimeSettings.h b/Engine/Source/Runtime/IOS/IOSRuntimeSettings/Classes/IOSRuntimeSettings.h
--- a/Engine/Source/Runtime/IOS/IOSRuntimeSettings/Classes/IOSRuntimeSettings.h
+++ b/Engine/Source/Runtime/IOS/IOSRuntimeSettings/Classes/IOSRuntimeSettings.h
@@ -179,11 +179,11 @@ public:
 
 	// Should Game Center support (iOS Online Subsystem) be enabled?
 	UPROPERTY(GlobalConfig, EditAnywhere, Category = Online, meta = (ConfigHierarchyEditable))
-	uint32 bEnableGameCenterSupport : 1;
+    bool bEnableGameCenterSupport;
 	
 	// Should Cloud Kit support (iOS Online Subsystem) be enabled?
 	UPROPERTY(GlobalConfig, EditAnywhere, Category = Online)
-	uint32 bEnableCloudKitSupport : 1;
+	bool bEnableCloudKitSupport;
 
 	// iCloud Read stategy
 	UPROPERTY(GlobalConfig, EditAnywhere, Category = Online, meta = (DisplayName = "iCloud save files sync strategy"), meta = (EditCondition = "bEnableCloudKitSupport"))
@@ -191,11 +191,11 @@ public:
 
     // Should push/remote notifications support (iOS Online Subsystem) be enabled?
     UPROPERTY(GlobalConfig, EditAnywhere, Category = Online)
-    uint32 bEnableRemoteNotificationsSupport : 1;
+    bool bEnableRemoteNotificationsSupport;
     
     // Should background fetch support be enabled?
     UPROPERTY(GlobalConfig, EditAnywhere, Category = Online)
-    uint32 bEnableBackgroundFetch : 1;
+    bool bEnableBackgroundFetch;
     
 	// Whether or not to compile iOS Metal shaders for the Mobile renderer (requires iOS 8+ and an A7 processor).
 	UPROPERTY(GlobalConfig, EditAnywhere, Category = Rendering, meta = (DisplayName = "Metal Mobile Renderer"))
@@ -206,13 +206,17 @@ public:
 	bool bSupportsMetalMRT;
     
     // Should the app be compatible for high refresh rate (iPhone only)
-    UPROPERTY(GlobalConfig, EditAnywhere, Category = "Rendering", meta = (DisplayName = "Enable ProMotion 120Hz on supported iPhone devices"))
+    UPROPERTY(GlobalConfig, EditAnywhere, Category = Rendering, meta = (DisplayName = "Enable ProMotion 120Hz on supported iPhone devices"))
     bool bSupportHighRefreshRates;
+        
+    /** Whether to enable LOD streaming for landscape visual meshes. Requires Metal support. */
+    UPROPERTY(GlobalConfig, EditAnywhere, Category = Rendering, Meta = (DisplayName = "Stream landscape visual mesh LODs"))
+    bool bStreamLandscapeMeshLODs;
+    
+    // Minimum iOS version this game supports
+    UPROPERTY(GlobalConfig, EditAnywhere, Category = Build, meta = (DisplayName = "Minimum iOS Version"))
+    EIOSVersion MinimumiOSVersion;
 	
-	// Support a secondary remote mac ?
-	UPROPERTY(GlobalConfig, EditAnywhere, Category = "Build", meta = (DisplayName = "Supports a secondary remote mac to facilitate debug of iOS / tvOS builds"))
-	bool bSupportSecondaryMac;
-
     // Whether to build the iOS project as a framework.
     UPROPERTY(GlobalConfig, EditAnywhere, Category = Build, meta = (DisplayName = "Build project as a framework (Experimental)"))
     bool bBuildAsFramework;
@@ -247,71 +251,85 @@ public:
 	// Any additional linker flags to pass to the linker in shipping builds
 	UPROPERTY(GlobalConfig, EditAnywhere, Category = Build, meta = (DisplayName = "Additional Shipping Linker Flags", ConfigHierarchyEditable))
 	FString AdditionalShippingLinkerFlags;
+    
+    // Any additional plist key/value data utilizing \n for a new line
+    UPROPERTY(GlobalConfig, EditAnywhere, Category = Build)
+    FString AdditionalPlistData;
+    
+    // Whether or not to add support for iPad devices
+    UPROPERTY(GlobalConfig, EditAnywhere, Category = Build, meta = (DisplayName = "Supports iPad"))
+    bool bSupportsIPad;
 
-	// The name or ip address of the remote mac which will be used to build IOS
-	UPROPERTY(GlobalConfig, EditAnywhere, Category = "Build", meta = (ConfigHierarchyEditable))
-	FString RemoteServerName;
-
-	
-	
-
-	// Enable the use of RSync for remote builds on a mac
-	UPROPERTY(GlobalConfig, EditAnywhere, Category = "Build", meta = (DisplayName = "Use RSync for building IOS", ConfigHierarchyEditable))
-	bool bUseRSync;
-
-	// The mac users name which matches the SSH Private Key, for remote builds using RSync.
-	UPROPERTY(GlobalConfig, EditAnywhere, Category = "Build", meta = (EditCondition = "bUseRSync", DisplayName = "Username on Remote Server.", ConfigHierarchyEditable))
-	FString RSyncUsername;
-
+    // Whether or not to add support for iPhone devices
+    UPROPERTY(GlobalConfig, EditAnywhere, Category = Build, meta = (DisplayName = "Supports iPhone"))
+    bool bSupportsIPhone;
+    
+    /** Set the maximum frame rate to save on power consumption */
+    UPROPERTY(GlobalConfig, EditAnywhere, Category = PowerUsage, meta = (ConfigHierarchyEditable))
+    EPowerUsageFrameRateLock FrameRateLock;
 
+    //Whether or not to allow taking the MaxRefreshRate from the device instead of a constant (60fps) in IOSPlatformFramePacer
+    UPROPERTY(GlobalConfig, EditAnywhere, Category = PowerUsage, meta = (ConfigHierarchyEditable))
+    bool bEnableDynamicMaxFPS;
 
-	// Optional path on the remote mac where the build files will be copied. If blank, ~/UE5/Builds will be used.
-	UPROPERTY(GlobalConfig, EditAnywhere, Category = "Build", meta = (ConfigHierarchyEditable))
-	FString RemoteServerOverrideBuildPath;
+    // Enable the use of RSync for remote builds on a mac
+    UPROPERTY(GlobalConfig, EditAnywhere, Category = "Remote Build", meta = (DisplayName = "Use RSync for building IOS", ConfigHierarchyEditable))
+    bool bUseRSync;
 
-	
-	// The install directory of cwrsync.
-	UPROPERTY(GlobalConfig, EditAnywhere, Category = "Build", meta = (EditCondition = "bUseRSync", ConfigHierarchyEditable))
-	FIOSBuildResourceDirectory CwRsyncInstallPath;
+    // The name or ip address of the remote mac which will be used to build IOS
+    UPROPERTY(GlobalConfig, EditAnywhere, Category = "Remote Build", meta = (ConfigHierarchyEditable))
+    FString RemoteServerName;
 
-	// The existing location of an SSH Key found by Unreal Engine.
-	UPROPERTY(VisibleAnywhere, Category = "Build", meta = (DisplayName = "Found Existing SSH permissions file"))
-	FString SSHPrivateKeyLocation;
+    // The mac users name which matches the SSH Private Key, for remote builds using RSync.
+    UPROPERTY(GlobalConfig, EditAnywhere, Category = "Remote Build", meta = (EditCondition = "bUseRSync", DisplayName = "Username on Remote Server", ConfigHierarchyEditable))
+    FString RSyncUsername;
 
-	// The path of the ssh permissions key to be used when connecting to the remote server.
-	UPROPERTY(GlobalConfig, EditAnywhere, Category = "Build", meta = (EditCondition = "bUseRSync", DisplayName = "Override existing SSH permissions file", ConfigHierarchyEditable))
-	FIOSBuildResourceFilePath SSHPrivateKeyOverridePath;
-
-	// The name or ip address of the remote mac which will be used to build IOS
-	UPROPERTY(GlobalConfig, EditAnywhere, Category = "Build", meta = (EditCondition = "bSupportSecondaryMac", ConfigHierarchyEditable))
-	FString SecondaryRemoteServerName;
+    // Optional path on the remote mac where the build files will be copied. If blank, ~/UE5/Builds will be used.
+    UPROPERTY(GlobalConfig, EditAnywhere, Category = "Remote Build", meta = (ConfigHierarchyEditable))
+    FString RemoteServerOverrideBuildPath;
+    
+    // The install directory of cwrsync.
+    UPROPERTY(GlobalConfig, EditAnywhere, Category = "Remote Build", meta = (EditCondition = "bUseRSync", ConfigHierarchyEditable))
+    FIOSBuildResourceDirectory CwRsyncInstallPath;
 
-	// Optional path on the secondary remote mac where the build files will be copied. If blank, ~/UE5/Builds will be used.
-	UPROPERTY(GlobalConfig, EditAnywhere, Category = "Build", meta = (EditCondition = "bSupportSecondaryMac", ConfigHierarchyEditable))
-	FString SecondaryRemoteServerOverrideBuildPath;
+    // The existing location of an SSH Key found by Unreal Engine.
+    UPROPERTY(VisibleAnywhere, Category = "Remote Build", meta = (DisplayName = "Found Existing SSH permissions file"))
+    FString SSHPrivateKeyLocation;
 
-	// The existing location of an SSH Key found by Unreal Engine.
-	UPROPERTY(VisibleAnywhere, Category = "Build", meta = (EditCondition = "bSupportSecondaryMac", DisplayName = "Found Existing SSH permissions file for Secondary Mac"))
-	FString SecondarySSHPrivateKeyLocation;
+    // The path of the ssh permissions key to be used when connecting to the remote server.
+    UPROPERTY(GlobalConfig, EditAnywhere, Category = "Remote Build", meta = (EditCondition = "bUseRSync", DisplayName = "Override existing SSH permissions file", ConfigHierarchyEditable))
+    FIOSBuildResourceFilePath SSHPrivateKeyOverridePath;
 
-	// The path of the ssh permissions key to be used when connecting to the remote server.
-	UPROPERTY(GlobalConfig, EditAnywhere, Category = "Build", meta = (EditCondition = "bSupportSecondaryMac", DisplayName = "Override existing SSH permissions file for Secondary Mac", ConfigHierarchyEditable))
-	FIOSBuildResourceFilePath SecondarySSHPrivateKeyOverridePath;
+    // Support a secondary remote Mac to support to facilitate iOS/tvOS debug ?
+    UPROPERTY(GlobalConfig, EditAnywhere, Category = "Remote Build", meta = (DisplayName = "Enable Secondary remote Mac"))
+    bool bSupportSecondaryMac;
+    
+    // The name or ip address of the remote mac which will be used to build IOS
+    UPROPERTY(GlobalConfig, EditAnywhere, Category = "Remote Build", meta = (EditCondition = "bSupportSecondaryMac", ConfigHierarchyEditable))
+    FString SecondaryRemoteServerName;
+    
+    // The secondary mac users name which matches the SSH Private Key, for remote builds using RSync.
+    UPROPERTY(GlobalConfig, EditAnywhere, Category = "Remote Build", meta = (EditCondition = "bSupportSecondaryMac", DisplayName = "Username on Secondary Remote Server", ConfigHierarchyEditable))
+    FString SecondaryRSyncUsername;
 
-	// The install directory of cwrsync.
-	UPROPERTY(GlobalConfig, EditAnywhere, Category = "Build", meta = (EditCondition = "bSupportSecondaryMac", ConfigHierarchyEditable))
-	FIOSBuildResourceDirectory SecondaryCwRsyncInstallPath;
+    // Optional path on the secondary remote mac where the build files will be copied. If blank, ~/UE5/Builds will be used.
+    UPROPERTY(GlobalConfig, EditAnywhere, Category = "Remote Build", meta = (EditCondition = "bSupportSecondaryMac", ConfigHierarchyEditable))
+    FString SecondaryRemoteServerOverrideBuildPath;
 
-	// The secondary mac users name which matches the SSH Private Key, for remote builds using RSync.
-	UPROPERTY(GlobalConfig, EditAnywhere, Category = "Build", meta = (EditCondition = "bSupportSecondaryMac", DisplayName = "Username on Remote Server.", ConfigHierarchyEditable))
-	FString SecondaryRSyncUsername;
+    // The install directory of cwrsync.
+    UPROPERTY(GlobalConfig, EditAnywhere, Category = "Remote Build", meta = (EditCondition = "bSupportSecondaryMac", ConfigHierarchyEditable))
+    FIOSBuildResourceDirectory SecondaryCwRsyncInstallPath;
+    
+    // The existing location of an SSH Key found by Unreal Engine.
+    UPROPERTY(VisibleAnywhere, Category = "Remote Build", meta = (EditCondition = "bSupportSecondaryMac", DisplayName = "Found Existing SSH permissions file for Secondary Mac"))
+    FString SecondarySSHPrivateKeyLocation;
 
-	// Should the app be compatible with Multi-User feature on tvOS ?　If checked, the game will will shutdown with the typical exit flow.
-    UPROPERTY(GlobalConfig, EditAnywhere, Category = "Build", meta = (DisplayName = "Support user switching on tvOS"))
-    bool bRunAsCurrentUser;
+    // The path of the ssh permissions key to be used when connecting to the remote server.
+    UPROPERTY(GlobalConfig, EditAnywhere, Category = "Remote Build", meta = (EditCondition = "bSupportSecondaryMac", DisplayName = "Override existing SSH permissions file for Secondary Mac", ConfigHierarchyEditable))
+    FIOSBuildResourceFilePath SecondarySSHPrivateKeyOverridePath;
 
-	// If checked, the game will be able to handle multiple gamepads at the same time (the Siri Remote is a gamepad)
-	UPROPERTY(GlobalConfig, EditAnywhere, Category = Input, meta = (DisplayName = "Can the Game have multiple gamepads connected at a single time"))
+    // If checked, the game will be able to handle multiple gamepads at the same time (the Siri Remote is a gamepad)
+	UPROPERTY(GlobalConfig, EditAnywhere, Category = Input, meta = (DisplayName = "Multiple gamepads support"))
 	bool bGameSupportsMultipleActiveControllers;
 
 	// If checked, the Siri Remote can be rotated to landscape view
@@ -374,30 +392,6 @@ public:
 	UPROPERTY(GlobalConfig, EditAnywhere, Category = BundleInformation)
 	FString VersionInfo;
 
-	/** Set the maximum frame rate to save on power consumption */
-	UPROPERTY(GlobalConfig, EditAnywhere, Category = PowerUsage, meta = (ConfigHierarchyEditable))
-	EPowerUsageFrameRateLock FrameRateLock;
-
-	//Whether or not to allow taking the MaxRefreshRate from the device instead of a constant (60fps) in IOSPlatformFramePacer
-	UPROPERTY(GlobalConfig, EditAnywhere, Category = PowerUsage, meta = (ConfigHierarchyEditable))
-	bool bEnableDynamicMaxFPS;
-
-	// Minimum iOS version this game supports
-	UPROPERTY(GlobalConfig, EditAnywhere, Category = OSInfo, meta = (DisplayName = "Minimum iOS Version"))
-	EIOSVersion MinimumiOSVersion;
-
-	// Whether or not to add support for iPad devices
-	UPROPERTY(GlobalConfig, EditAnywhere, Category = DeviceUsage)
-	uint32 bSupportsIPad : 1;
-
-	// Whether or not to add support for iPhone devices
-	UPROPERTY(GlobalConfig, EditAnywhere, Category = DeviceUsage)
-	uint32 bSupportsIPhone : 1;
-
-	// Any additional plist key/value data utilizing \n for a new line
-	UPROPERTY(GlobalConfig, EditAnywhere, Category = ExtraData)
-	FString AdditionalPlistData;
-
 	/**
 	 * Choose whether to use a custom LaunchScreen.Storyboard as a Launchscreen. To use this option, create a storyboard in Xcode and 
 	 * copy it named LaunchScreen.storyboard in Build/IOS/Resources/Interface under your Project folder. This will be compiled and 
@@ -537,8 +531,8 @@ public:
 	UPROPERTY(config, EditAnywhere, Category = "Audio")
 	FPlatformRuntimeAudioCompressionOverrides CompressionOverrides;
 
-    /** Whether this app's audio can be played when using other apps or on the srpingboard */
-    UPROPERTY(config, EditAnywhere, Category = "Audio", meta = (DisplayName = "Whether audio from this plays in background (iOS and iPadOS only)"))
+    /** Whether this app's audio can be played when using other apps or on the springboard */
+    UPROPERTY(config, EditAnywhere, Category = "Audio", meta = (DisplayName = "Enable Background Audio"))
     bool bSupportsBackgroundAudio;
 
 	/** This determines the max amount of memory that should be used for the cache at any given time. If set low (<= 8 MB), it lowers the size of individual chunks of audio during cook. */
@@ -581,11 +575,7 @@ public:
 	UPROPERTY(GlobalConfig)
 	float AutoStreamingThreshold;
 
-	/** Whether to enable LOD streaming for landscape visual meshes. Requires Metal support. */
-	UPROPERTY(GlobalConfig, EditAnywhere, Category = "Misc", Meta = (DisplayName = "Stream landscape visual mesh LODs"))
-	bool bStreamLandscapeMeshLODs;
-
-	virtual void PostReloadConfig(class FProperty* PropertyThatWasLoaded) override;
+    virtual void PostReloadConfig(class FProperty* PropertyThatWasLoaded) override;
 
 #if WITH_EDITOR
 	// UObject interface
diff --git a/Engine/Source/Runtime/IOS/IOSRuntimeSettings/Private/IOSRuntimeSettings.cpp b/Engine/Source/Runtime/IOS/IOSRuntimeSettings/Private/IOSRuntimeSettings.cpp
--- a/Engine/Source/Runtime/IOS/IOSRuntimeSettings/Private/IOSRuntimeSettings.cpp
+++ b/Engine/Source/Runtime/IOS/IOSRuntimeSettings/Private/IOSRuntimeSettings.cpp
@@ -19,7 +19,6 @@ UIOSRuntimeSettings::UIOSRuntimeSettings(const FObjectInitializer& ObjectInitial
 {
 	bEnableGameCenterSupport = true;
 	bEnableCloudKitSupport = false;
-    bRunAsCurrentUser = false;
 	bSupportsPortraitOrientation = true;
 	bSupportsITunesFileSharing = false;
 	bSupportsFilesApp = false;
