Fix tvOS multi user

diff --git a/Engine/Config/BaseEngine.ini b/Engine/Config/BaseEngine.ini
--- a/Engine/Config/BaseEngine.ini
+++ b/Engine/Config/BaseEngine.ini
@@ -2874,6 +2874,7 @@ SSHPrivateKeyOverridePath=
 SecondarySSHPrivateKeyOverridePath=
 bEnableRemoteNotificationsSupport=False
 bEnableCloudKitSupport=False
+bUserSwitching=False
 bSupportSecondaryMac=False
 IOSCloudKitSyncStrategy=None
 bGenerateCrashReportSymbols=false
diff --git a/Engine/Source/Programs/UnrealBuildTool/Platform/IOS/IOSExports.cs b/Engine/Source/Programs/UnrealBuildTool/Platform/IOS/IOSExports.cs
--- a/Engine/Source/Programs/UnrealBuildTool/Platform/IOS/IOSExports.cs
+++ b/Engine/Source/Programs/UnrealBuildTool/Platform/IOS/IOSExports.cs
@@ -333,7 +333,9 @@ namespace UnrealBuildTool
 				}
 
 				// Add Multi-user support for tvOS
-				if (Platform == UnrealTargetPlatform.TVOS)
+				bool bUserSwitching = false;
+				PlatformGameConfig.GetBool("/Script/IOSRuntimeSettings.IOSRuntimeSettings", "bUserSwitching", out bUserSwitching);
+				if (bUserSwitching && Platform == UnrealTargetPlatform.TVOS)
 				{
 					Text.AppendLine("\t<key>com.apple.developer.user-management</key>");
 					Text.AppendLine("\t<array><string>runs-as-current-user</string></array>");
diff --git a/Engine/Source/Runtime/IOS/IOSRuntimeSettings/Classes/IOSRuntimeSettings.h b/Engine/Source/Runtime/IOS/IOSRuntimeSettings/Classes/IOSRuntimeSettings.h
--- a/Engine/Source/Runtime/IOS/IOSRuntimeSettings/Classes/IOSRuntimeSettings.h
+++ b/Engine/Source/Runtime/IOS/IOSRuntimeSettings/Classes/IOSRuntimeSettings.h
@@ -326,6 +326,10 @@ public:
     UPROPERTY(GlobalConfig, EditAnywhere, Category = "Remote Build", meta = (EditCondition = "bSupportSecondaryMac", DisplayName = "Override existing SSH permissions file for Secondary Mac", ConfigHierarchyEditable))
     FIOSBuildResourceFilePath SecondarySSHPrivateKeyOverridePath;
 
+	// Should the app be multi-users compatible on tvOS ? Requires the com.apple.developer.user-management entitlement.
+    UPROPERTY(GlobalConfig, EditAnywhere, Category = "Build", meta = (DisplayName = "Support user switching on tvOS"))
+    bool bUserSwitching;
+
     // If checked, the game will be able to handle multiple gamepads at the same time (the Siri Remote is a gamepad)
 	UPROPERTY(GlobalConfig, EditAnywhere, Category = Input, meta = (DisplayName = "Multiple gamepads support"))
 	bool bGameSupportsMultipleActiveControllers;
diff --git a/Engine/Source/Runtime/IOS/IOSRuntimeSettings/Private/IOSRuntimeSettings.cpp b/Engine/Source/Runtime/IOS/IOSRuntimeSettings/Private/IOSRuntimeSettings.cpp
--- a/Engine/Source/Runtime/IOS/IOSRuntimeSettings/Private/IOSRuntimeSettings.cpp
+++ b/Engine/Source/Runtime/IOS/IOSRuntimeSettings/Private/IOSRuntimeSettings.cpp
@@ -19,6 +19,7 @@ UIOSRuntimeSettings::UIOSRuntimeSettings(const FObjectInitializer& ObjectInitial
 {
 	bEnableGameCenterSupport = true;
 	bEnableCloudKitSupport = false;
+	bUserSwitching = false;
 	bSupportsPortraitOrientation = true;
 	bSupportsITunesFileSharing = false;
 	bSupportsFilesApp = false;
