When toggling iOS sim support, project files need rebuilding

diff --git a/Engine/Source/Runtime/IOS/IOSRuntimeSettings/Private/IOSRuntimeSettings.cpp b/Engine/Source/Runtime/IOS/IOSRuntimeSettings/Private/IOSRuntimeSettings.cpp
--- a/Engine/Source/Runtime/IOS/IOSRuntimeSettings/Private/IOSRuntimeSettings.cpp
+++ b/Engine/Source/Runtime/IOS/IOSRuntimeSettings/Private/IOSRuntimeSettings.cpp
@@ -7,6 +7,14 @@
 #include "HAL/IConsoleManager.h"
 #include "HAL/PlatformApplicationMisc.h"
 
+#if WITH_EDITOR
+#include "Misc/App.h"
+#include "Misc/MessageDialog.h"
+#include "Misc/ScopedSlowTask.h"
+#include "DesktopPlatformModule.h"
+#define LOCTEXT_NAMESPACE "IOSRuntimeSettings"
+#endif
+
 UIOSRuntimeSettings::UIOSRuntimeSettings(const FObjectInitializer& ObjectInitializer)
 	: Super(ObjectInitializer)
 	, CacheSizeKB(65536)
@@ -86,6 +94,36 @@ void UIOSRuntimeSettings::PostEditChangeProperty(struct FPropertyChangedEvent& P
 		UpdateSinglePropertyInConfigFile(GetClass()->FindPropertyByName(GET_MEMBER_NAME_CHECKED(UIOSRuntimeSettings, bSupportsMetal)), GetDefaultConfigFilename());
 	}
 
+	// If iOS Simulator setting changed, need to rerun GPF to force xcconfig files to updated the supported platforms
+	if (PropertyChangedEvent.GetPropertyName() == TEXT("bEnableSimulatorSupport") &&
+		PropertyChangedEvent.ChangeType == EPropertyChangeType::ValueSet)
+	{
+		const FText GPFQueryLoc = LOCTEXT("RunGPFQuery", "Xcode Project will refresh. Continue?");
+		if (FMessageDialog::Open(EAppMsgType::OkCancel, GPFQueryLoc) == EAppReturnType::Cancel)
+		{
+			// User canceled, so reset the value
+			bEnableSimulatorSupport = !bEnableSimulatorSupport;
+			UpdateSinglePropertyInConfigFile(GetClass()->FindPropertyByName(GET_MEMBER_NAME_CHECKED(UIOSRuntimeSettings, bEnableSimulatorSupport)), GetDefaultConfigFilename());
+			return;
+		}
+
+		FScopedSlowTask SlowTask(0, LOCTEXT("UpdatingCodeProject", "Updating code project..."));
+		SlowTask.MakeDialog();
+
+		// Try to generate project files
+		FStringOutputDevice OutputLog;
+		OutputLog.SetAutoEmitLineTerminator(true);
+		GLog->AddOutputDevice(&OutputLog);
+		bool bSuccess = FDesktopPlatformModule::Get()->GenerateProjectFiles(FPaths::RootDir(), FPaths::ProjectDir() + FApp::GetProjectName(), GWarn);
+		GLog->RemoveOutputDevice(&OutputLog);
+
+		if (!bSuccess)
+		{
+			FFormatNamedArguments Args;
+			Args.Add(TEXT("LogOutput"), FText::FromString(OutputLog));
+			FMessageDialog::Open(EAppMsgType::Ok, FText::Format(LOCTEXT("CouldNotGenerateProjectFiles", "Project files could not be generated. Please run GPF manually or revert the setting. Log:\n\n{LogOutput}"), Args));
+		}
+	}
 }
 
 
