Allow setting of ITSAppUsesNonExemptEncryption and when set to "true", set ITSEncryptionExportComplianceCode in plist

diff --git a/Engine/Config/BaseEngine.ini b/Engine/Config/BaseEngine.ini
--- a/Engine/Config/BaseEngine.ini
+++ b/Engine/Config/BaseEngine.ini
@@ -3112,6 +3112,8 @@ ForceFloats=True
 EnableMathOptimisations=True
 MetalLanguageVersion=0
 bDisableMotionData=False
+bUsesNonExemptEncryption=false
+ITSEncryptionExportComplianceCode=""
 bEnableAdvertisingIdentifier=True
 
 ; These are the defaults for Android settings, and they need to be in the .ini since UBT reads the .ini settings, without instantiating the class
diff --git a/Engine/Source/Programs/UnrealBuildTool/Platform/IOS/UEDeployIOS.cs b/Engine/Source/Programs/UnrealBuildTool/Platform/IOS/UEDeployIOS.cs
--- a/Engine/Source/Programs/UnrealBuildTool/Platform/IOS/UEDeployIOS.cs
+++ b/Engine/Source/Programs/UnrealBuildTool/Platform/IOS/UEDeployIOS.cs
@@ -410,9 +410,19 @@ namespace UnrealBuildTool
 				Text.AppendLine("\t<key>CADisableMinimumFrameDurationOnPhone</key><true/>");
 			}
 
-			// disable exempt encryption
+			// set exempt encryption
+			bool bUsesNonExemptEncryption = false;
+			string ITSEncryptionExportComplianceCode = "";
+			Ini.GetBool("/Script/IOSRuntimeSettings.IOSRuntimeSettings", "bUsesNonExemptEncryption", out bUsesNonExemptEncryption);
+			Ini.GetString("/Script/IOSRuntimeSettings.IOSRuntimeSettings", "ITSEncryptionExportComplianceCode", out ITSEncryptionExportComplianceCode);
 			Text.AppendLine("\t<key>ITSAppUsesNonExemptEncryption</key>");
-			Text.AppendLine("\t<false/>");
+			Text.AppendLine(String.Format("\t<{0}/>", bUsesNonExemptEncryption ? "true" : "false"));
+			if (bUsesNonExemptEncryption && !String.IsNullOrWhiteSpace(ITSEncryptionExportComplianceCode))
+			{
+				Text.AppendLine("\t<key>ITSEncryptionExportComplianceCode</key>");
+				Text.AppendLine(String.Format("\t<string>{0}</string>", ITSEncryptionExportComplianceCode));
+			}
+			
 			// add location services descriptions if used
 			if (!String.IsNullOrWhiteSpace(LocationAlwaysUsageDescription))
 			{
diff --git a/Engine/Source/Runtime/IOS/IOSRuntimeSettings/Classes/IOSRuntimeSettings.h b/Engine/Source/Runtime/IOS/IOSRuntimeSettings/Classes/IOSRuntimeSettings.h
--- a/Engine/Source/Runtime/IOS/IOSRuntimeSettings/Classes/IOSRuntimeSettings.h
+++ b/Engine/Source/Runtime/IOS/IOSRuntimeSettings/Classes/IOSRuntimeSettings.h
@@ -214,7 +214,7 @@ public:
     // Should the app be compatible for high refresh rate (iPhone only)
     UPROPERTY(GlobalConfig, EditAnywhere, Category = Rendering, meta = (DisplayName = "Enable ProMotion 120Hz on supported iPhone devices"))
     bool bSupportHighRefreshRates;
-        
+
     /** Whether to enable LOD streaming for landscape visual meshes. Requires Metal support. */
     UPROPERTY(GlobalConfig, EditAnywhere, Category = Rendering, Meta = (DisplayName = "Stream landscape visual mesh LODs"))
     bool bStreamLandscapeMeshLODs;
@@ -244,12 +244,20 @@ public:
 	
 	// Enable generation of xcode archive package
 	UPROPERTY(GlobalConfig, EditAnywhere, Category = Build, meta = (DisplayName = "Generate xcode archive package"))
-	bool bGenerateXCArchive;	
-	
+	bool bGenerateXCArchive;
+
+	// Does the app require use of NON-standard encryption
+	UPROPERTY(GlobalConfig, EditAnywhere, Category = Build, meta = (DisplayName = "ITSAppUsesNonExemptEncryption plist value"))
+	bool bUsesNonExemptEncryption;
+
+	// If bUsesNonExemptEncryption, set the ITSEncryptionExportComplianceCode
+	UPROPERTY(GlobalConfig, EditAnywhere, Category = Build, meta = (EditCondition = "bUsesNonExemptEncryption", DisplayName = "Set ITSEncryptionExportComplianceCode"))
+	FString ITSEncryptionExportComplianceCode;
+
 	// Enable Advertising Identified
 	UPROPERTY(GlobalConfig, EditAnywhere, Category = Build, meta = (DisplayName = "Enable Advertising Identified (IDFA)"))
 	bool bEnableAdvertisingIdentifier;
-	
+
 	// Any additional linker flags to pass to the linker in non-shipping builds
 	UPROPERTY(GlobalConfig, EditAnywhere, Category = Build, meta = (DisplayName = "Additional Non-Shipping Linker Flags", ConfigHierarchyEditable))
 	FString AdditionalLinkerFlags;
diff --git a/Engine/Source/Runtime/IOS/IOSRuntimeSettings/Private/IOSRuntimeSettings.cpp b/Engine/Source/Runtime/IOS/IOSRuntimeSettings/Private/IOSRuntimeSettings.cpp
--- a/Engine/Source/Runtime/IOS/IOSRuntimeSettings/Private/IOSRuntimeSettings.cpp
+++ b/Engine/Source/Runtime/IOS/IOSRuntimeSettings/Private/IOSRuntimeSettings.cpp
@@ -47,6 +47,8 @@ UIOSRuntimeSettings::UIOSRuntimeSettings(const FObjectInitializer& ObjectInitial
 	bGeneratedSYMFile = false;
 	bGeneratedSYMBundle = false;
 	bGenerateXCArchive = false;
+	bUsesNonExemptEncryption = false;
+	ITSEncryptionExportComplianceCode = TEXT("");
 	bSupportSecondaryMac = false;
 	bUseRSync = true;
 	bCustomLaunchscreenStoryboard = false;
