Fix MarketplaceKit non-unity build

#rnx

diff --git a/Engine/Source/Runtime/IOS/MarketplaceKit/Public/MarketplaceKit.h b/Engine/Source/Runtime/IOS/MarketplaceKit/Public/MarketplaceKit.h
--- a/Engine/Source/Runtime/IOS/MarketplaceKit/Public/MarketplaceKit.h
+++ b/Engine/Source/Runtime/IOS/MarketplaceKit/Public/MarketplaceKit.h
@@ -2,6 +2,9 @@
 
 #pragma once
 
+#include "CoreMinimal.h"
+#include "Modules/ModuleManager.h"
+
 enum class EMarketplaceType : int32
 {
 	AppStore = 0,
