Implement ScheduleLocalNotificationAtTimeOverrideId on iOS

diff --git a/Engine/Source/Runtime/IOS/IOSLocalNotification/Private/IOSLocalNotification.cpp b/Engine/Source/Runtime/IOS/IOSLocalNotification/Private/IOSLocalNotification.cpp
--- a/Engine/Source/Runtime/IOS/IOSLocalNotification/Private/IOSLocalNotification.cpp
+++ b/Engine/Source/Runtime/IOS/IOSLocalNotification/Private/IOSLocalNotification.cpp
@@ -143,6 +143,16 @@ void FIOSLocalNotificationService::ClearAllLocalNotifications()
 static int32 NotificationNumber = 0;
 int32 FIOSLocalNotificationService::ScheduleLocalNotificationAtTime(const FDateTime& FireDateTime, bool LocalTime, const FText& Title, const FText& Body, const FText& Action, const FString& ActivationEvent)
 {
+#if !PLATFORM_TVOS
+	const int32 CurrentNotificationId = NotificationNumber++;
+	return ScheduleLocalNotificationAtTimeOverrideId(FireDateTime, LocalTime, Title, Body, Action, ActivationEvent, CurrentNotificationId);
+#else
+	return -1;
+#endif
+}
+
+int32 FIOSLocalNotificationService::ScheduleLocalNotificationAtTimeOverrideId(const FDateTime& FireDateTime, bool LocalTime, const FText& Title, const FText& Body, const FText& Action, const FString& ActivationEvent, int32 IdOverride)
+{
 #if !PLATFORM_TVOS
 	if (FireDateTime < (LocalTime ? FDateTime::Now() : FDateTime::UtcNow()))
 	{
@@ -155,8 +165,8 @@ int32 FIOSLocalNotificationService::ScheduleLocalNotificationAtTime(const FDateT
 	FText BodyCopy = Body;
 	FText ActionCopy = Action;
 	FString ActivationEventCopy = ActivationEvent;
-	int32 CurrentNotificationId = NotificationNumber++;
-	
+	const int32 CurrentNotificationId = IdOverride;
+
 	//have to schedule notification on main thread queue
 	dispatch_async(dispatch_get_main_queue(), ^{
 		UNMutableNotificationContent* Content = FIOSLocalNotificationModule::CreateNotificationContent(TitleCopy, BodyCopy, ActionCopy, ActivationEventCopy, 1);
@@ -178,11 +188,6 @@ int32 FIOSLocalNotificationService::ScheduleLocalNotificationAtTime(const FDateT
 #endif
 }
 
-int32 FIOSLocalNotificationService::ScheduleLocalNotificationAtTimeOverrideId(const FDateTime& FireDateTime, bool LocalTime, const FText& Title, const FText& Body, const FText& Action, const FString& ActivationEvent, int32 IdOverride)
-{
-	return ScheduleLocalNotificationAtTime(FireDateTime, LocalTime, Title, Body, Action, ActivationEvent);
-}
-
 int32 FIOSLocalNotificationService::ScheduleLocalNotificationBadgeAtTime(const FDateTime& FireDateTime, bool LocalTime, const FString& ActivationEvent)
 {
 #if !PLATFORM_TVOS
