Using MarketplaceKit.framework to open correct store on iOS

diff --git a/Engine/Source/Runtime/IOS/MarketplaceKit/Private/MarketplaceKit.cpp b/Engine/Source/Runtime/IOS/MarketplaceKit/Private/MarketplaceKit.cpp
--- a/Engine/Source/Runtime/IOS/MarketplaceKit/Private/MarketplaceKit.cpp
+++ b/Engine/Source/Runtime/IOS/MarketplaceKit/Private/MarketplaceKit.cpp
@@ -8,6 +8,7 @@ DEFINE_LOG_CATEGORY_STATIC(LogMarketplaceKit, Log, All);
 
 void FMarketplaceKitModule::StartupModule()
 {
+	CacheValue();
 }
 
 void FMarketplaceKitModule::ShutdownModule()
@@ -33,7 +34,7 @@ static constexpr EMarketplaceType ConvertMarketplaceType(const AppDistributorTyp
 	}
 }
 
-void FMarketplaceKitModule::GetCurrentType(TFunction<void(EMarketplaceType Type, const FString& Name)> Callback)
+void FMarketplaceKitModule::GetCurrentTypeAsync(TFunction<void(EMarketplaceType Type, const FString& Name)> Callback)
 {
 	[AppDistributorWrapper getCurrentWithCompletionHandler:^(enum AppDistributorType Type, NSString* _Nonnull Name)
 	{
@@ -46,4 +47,62 @@ void FMarketplaceKitModule::GetCurrentType(TFunction<void(EMarketplaceType Type,
 	}];
 }
 
+void FMarketplaceKitModule::GetCurrentType(EMarketplaceType& OutType, FString& OutName)
+{
+	CacheValue();
+	OutType = CachedType;
+	OutName = CachedName;
+}
+
+FString FMarketplaceKitModule::GetCurrentTypeAsString()
+{
+	CacheValue();
+	
+	TStringBuilder<256> Result;
+	switch (CachedType)
+	{
+		case EMarketplaceType::AppStore:		Result.Append(TEXT("AppStore")); break;
+		case EMarketplaceType::TestFlight:		Result.Append(TEXT("TestFlight")); break;
+		case EMarketplaceType::Marketplace:		Result.Append(TEXT("Marketplace")); break;
+		case EMarketplaceType::Web:				Result.Append(TEXT("Web")); break;
+		case EMarketplaceType::NotAvailable:	Result.Append(TEXT("NotAvailable")); break;
+		case EMarketplaceType::Other:			[[fallthrough]];
+		default:								Result.Append(TEXT("Other")); break;
+	}
+	
+	if (!CachedName.IsEmpty())
+	{
+		Result.Append(TEXT("-"));
+		Result.Append(CachedName);
+	}
+
+	return *Result;
+}
+
+void FMarketplaceKitModule::CacheValue()
+{
+	if (bCachedTypeValid)
+	{
+		return;
+	}
+
+	// TODO avoid scheduling multiple requests in case this path is hit from multiple threads
+
+	dispatch_semaphore_t Semaphore = dispatch_semaphore_create(0);
+
+	[AppDistributorWrapper getCurrentWithCompletionHandler:^(enum AppDistributorType Type, NSString* _Nonnull Name)
+	{
+		CachedType = ConvertMarketplaceType(Type);
+		CachedName = Name;
+		bCachedTypeValid = true;
+
+		UE_LOG(LogMarketplaceKit, Log, TEXT("AppDistributorWrapper getCurrentWithCompletionHandler %i %s"), (int32)CachedType, *CachedName);
+
+		dispatch_semaphore_signal(Semaphore);
+	}];
+
+	// wait for a result, but timeout after 1s
+	dispatch_semaphore_wait(Semaphore, dispatch_time(DISPATCH_TIME_NOW, NSEC_PER_SEC));
+}
+
 IMPLEMENT_MODULE(FMarketplaceKitModule, MarketplaceKit);
diff --git a/Engine/Source/Runtime/IOS/MarketplaceKit/Public/MarketplaceKit.h b/Engine/Source/Runtime/IOS/MarketplaceKit/Public/MarketplaceKit.h
--- a/Engine/Source/Runtime/IOS/MarketplaceKit/Public/MarketplaceKit.h
+++ b/Engine/Source/Runtime/IOS/MarketplaceKit/Public/MarketplaceKit.h
@@ -21,6 +21,15 @@ public:
 	virtual void StartupModule() override;
 	virtual void ShutdownModule() override;
 	virtual bool SupportsDynamicReloading() override;
-	
-	void GetCurrentType(TFunction<void(EMarketplaceType Type, const FString& Name)> Callback);
+
+	void GetCurrentTypeAsync(TFunction<void(EMarketplaceType Type, const FString& Name)> Callback);
+	void GetCurrentType(EMarketplaceType& OutType, FString& OutName);
+	FString GetCurrentTypeAsString();
+
+private:
+	bool bCachedTypeValid = false;
+	EMarketplaceType CachedType = EMarketplaceType::NotAvailable;
+	FString CachedName;
+
+	void CacheValue();
 };
