[CIS] Fix deprecation warnings on iOS:
Removing / deprecating stale code paths  (DTYPE_Native & DTYPE_Preview) and clawing back raw ptr-based resource management from FAudioDeviceManager.

#rnx

diff --git a/Engine/Source/Runtime/IOS/IOSAudio/Private/IOSAudioSource.cpp b/Engine/Source/Runtime/IOS/IOSAudio/Private/IOSAudioSource.cpp
--- a/Engine/Source/Runtime/IOS/IOSAudio/Private/IOSAudioSource.cpp
+++ b/Engine/Source/Runtime/IOS/IOSAudio/Private/IOSAudioSource.cpp
@@ -100,7 +100,6 @@ FIOSAudioSoundSource::~FIOSAudioSoundSource(void)
 		delete IOSBuffer;
 		IOSBuffer = NULL;
 	}
-	Buffer = NULL;
 }
 
 void FIOSAudioSoundSource::CleanupAudioBuffer()
@@ -136,7 +135,6 @@ bool FIOSAudioSoundSource::Init(FWaveInstance* InWaveInstance)
 	CleanupAudioBuffer();
 	
 	IOSBuffer = FIOSAudioSoundBuffer::Init(IOSAudioDevice, InWaveInstance->WaveData);
-	Buffer = IOSBuffer;
 
 	if (IOSBuffer == NULL || IOSBuffer->NumChannels <= 0 || (IOSBuffer->SoundFormat != SoundFormat_LPCM && IOSBuffer->SoundFormat != SoundFormat_ADPCM))
 	{
