CIS fix for IOSAudioBuffer
#rnx

diff --git a/Engine/Source/Runtime/IOS/IOSAudio/Private/IOSAudioBuffer.cpp b/Engine/Source/Runtime/IOS/IOSAudio/Private/IOSAudioBuffer.cpp
--- a/Engine/Source/Runtime/IOS/IOSAudio/Private/IOSAudioBuffer.cpp
+++ b/Engine/Source/Runtime/IOS/IOSAudio/Private/IOSAudioBuffer.cpp
@@ -116,7 +116,6 @@ int32 FIOSAudioSoundBuffer::GetCurrentChunkOffset() const
 bool FIOSAudioSoundBuffer::ReadCompressedInfo(USoundWave* InWave)
 {
 	check(DecompressionState != nullptr);
-	check(InWave->SoundWaveDataPtr.IsValid());
 
 	FSoundQualityInfo QualityInfo = { 0 };
 	if(bStreaming)
