iOS: Remove experimental feature from settings.

diff --git a/Engine/Source/Runtime/IOS/IOSRuntimeSettings/Classes/IOSRuntimeSettings.h b/Engine/Source/Runtime/IOS/IOSRuntimeSettings/Classes/IOSRuntimeSettings.h
--- a/Engine/Source/Runtime/IOS/IOSRuntimeSettings/Classes/IOSRuntimeSettings.h
+++ b/Engine/Source/Runtime/IOS/IOSRuntimeSettings/Classes/IOSRuntimeSettings.h
@@ -222,10 +222,6 @@ public:
     // Minimum iOS version this game supports
     UPROPERTY(GlobalConfig, EditAnywhere, Category = Build, meta = (DisplayName = "Minimum iOS Version"))
     EIOSVersion MinimumiOSVersion;
-	
-    // Whether to build the iOS project as a framework.
-    UPROPERTY(GlobalConfig, EditAnywhere, Category = Build, meta = (DisplayName = "Build project as a framework (Experimental)"))
-    bool bBuildAsFramework;
 
 	UPROPERTY(GlobalConfig, EditAnywhere, Category = Build, meta = (DisplayName = "Override location of Metal toolchain"))
 	FIOSBuildResourceDirectory WindowsMetalToolchainOverride;
diff --git a/Engine/Source/Runtime/IOS/IOSRuntimeSettings/Private/IOSRuntimeSettings.cpp b/Engine/Source/Runtime/IOS/IOSRuntimeSettings/Private/IOSRuntimeSettings.cpp
--- a/Engine/Source/Runtime/IOS/IOSRuntimeSettings/Private/IOSRuntimeSettings.cpp
+++ b/Engine/Source/Runtime/IOS/IOSRuntimeSettings/Private/IOSRuntimeSettings.cpp
@@ -43,7 +43,6 @@ UIOSRuntimeSettings::UIOSRuntimeSettings(const FObjectInitializer& ObjectInitial
 	bEnableSplitView = false;
 	bEnableSimulatorSupport = false;
 	MinimumiOSVersion = EIOSVersion::IOS_Minimum;
-    bBuildAsFramework = true;
 	bGeneratedSYMFile = false;
 	bGeneratedSYMBundle = false;
 	bGenerateXCArchive = false;
