- Add some static getters to FMarketplaceKitModule to reduce verbosity at callsites
- Add ability to set the Type/BundleId via CVar when running from testflight/xcode
- Make StoreKit features conditional on MarketplaceType==AppStore

diff --git a/Engine/Plugins/Online/IOS/OnlineSubsystemIOS/Source/OnlineSubsystemIOS.Build.cs b/Engine/Plugins/Online/IOS/OnlineSubsystemIOS/Source/OnlineSubsystemIOS.Build.cs
--- a/Engine/Plugins/Online/IOS/OnlineSubsystemIOS/Source/OnlineSubsystemIOS.Build.cs
+++ b/Engine/Plugins/Online/IOS/OnlineSubsystemIOS/Source/OnlineSubsystemIOS.Build.cs
@@ -9,7 +9,7 @@ public class OnlineSubsystemIOS : ModuleRules
 		PCHUsage = ModuleRules.PCHUsageMode.UseExplicitOrSharedPCHs;
 		PrivateIncludePaths.AddRange( 
             new string[] {
-                "Private",             
+                "Private",
                 });
 
         PublicIncludePaths.AddRange(
@@ -24,10 +24,11 @@ public class OnlineSubsystemIOS : ModuleRules
                 "Core", 
                 "CoreUObject", 
                 "Engine", 
-                "Sockets",
-				"OnlineSubsystem", 
                 "HTTP",
                 "IOSPlatformFeatures",
+                "MarketplaceKit",
+				"OnlineSubsystem", 
+                "Sockets",
             }
             );
 
diff --git a/Engine/Plugins/Online/IOS/OnlineSubsystemIOS/Source/Private/OnlineSubsystemIOS.cpp b/Engine/Plugins/Online/IOS/OnlineSubsystemIOS/Source/Private/OnlineSubsystemIOS.cpp
--- a/Engine/Plugins/Online/IOS/OnlineSubsystemIOS/Source/Private/OnlineSubsystemIOS.cpp
+++ b/Engine/Plugins/Online/IOS/OnlineSubsystemIOS/Source/Private/OnlineSubsystemIOS.cpp
@@ -1,9 +1,11 @@
 // Copyright Epic Games, Inc. All Rights Reserved.
 
 #include "OnlineSubsystemIOS.h"
+
 #include "IOS/IOSAppDelegate.h"
-#import "OnlineAppStoreUtils.h"
+#include "MarketplaceKit.h"
 #include "Misc/ConfigCacheIni.h"
+#import "OnlineAppStoreUtils.h"
 #include "Stats/Stats.h"
 
 FOnlineSubsystemIOS::FOnlineSubsystemIOS(FName InInstanceName)
@@ -167,8 +169,11 @@ bool FOnlineSubsystemIOS::Init()
 #if UE_WITH_STORE_KIT
 		if (IsInAppPurchasingEnabled())
 		{
-			StoreV2Interface = MakeShareable(new FOnlineStoreIOS(this));
-			PurchaseInterface = MakeShareable(new FOnlinePurchaseIOS(this));
+			if (FMarketplaceKitModule::GetCurrentTypeStatic() == EMarketplaceType::AppStore)
+			{
+				StoreV2Interface = MakeShareable(new FOnlineStoreIOS(this));
+				PurchaseInterface = MakeShareable(new FOnlinePurchaseIOS(this));
+			}
 		}
 #endif
 		if (UserCloudInterface && IsCloudKitEnabled())
diff --git a/Engine/Source/Runtime/Core/Core.Build.cs b/Engine/Source/Runtime/Core/Core.Build.cs
--- a/Engine/Source/Runtime/Core/Core.Build.cs
+++ b/Engine/Source/Runtime/Core/Core.Build.cs
@@ -42,6 +42,7 @@ public class Core : ModuleRules
 		PrivateIncludePathModuleNames.AddRange(
 			new string[] {
 				"DerivedDataCache",
+				"MarketplaceKit",
 				"TargetPlatform",
 
 				"Json",
diff --git a/Engine/Source/Runtime/Core/Private/IOS/IOSPlatformMisc.cpp b/Engine/Source/Runtime/Core/Private/IOS/IOSPlatformMisc.cpp
--- a/Engine/Source/Runtime/Core/Private/IOS/IOSPlatformMisc.cpp
+++ b/Engine/Source/Runtime/Core/Private/IOS/IOSPlatformMisc.cpp
@@ -20,6 +20,7 @@
 #include "IOS/IOSPlatformCrashContext.h"
 #include "IOS/IOSPlatformPLCrashReporterIncludes.h"
 #include "IOS/IOSView.h"
+#include "MarketplaceKit.h"
 #include "Misc/App.h"
 #include "Misc/CommandLine.h"
 #include "Misc/ConfigCacheIni.h"
@@ -981,10 +982,13 @@ bool FIOSPlatformMisc::GetDiskTotalAndFreeSpace(const FString& InPath, uint64& T
 void FIOSPlatformMisc::RequestStoreReview()
 {
 #if UE_WITH_STORE_KIT && !PLATFORM_TVOS
-PRAGMA_DISABLE_DEPRECATION_WARNINGS
-	// Deprecated for iOS18.0 & visionOS 2.0.  To be replaced with Swift only StoreKit::RequestReviewAction call (UE-228925)
-    [SKStoreReviewController requestReviewInScene:[[[[UIApplication sharedApplication] delegate] window] windowScene]];
-PRAGMA_ENABLE_DEPRECATION_WARNINGS
+	if (FMarketplaceKitModule::GetCurrentTypeStatic() == EMarketplaceType::AppStore)
+	{
+		// Deprecated for iOS18.0 & visionOS 2.0.  To be replaced with Swift only StoreKit::RequestReviewAction call (UE-228925)
+		PRAGMA_DISABLE_DEPRECATION_WARNINGS
+		[SKStoreReviewController requestReviewInScene:[[[[UIApplication sharedApplication] delegate] window] windowScene]];
+		PRAGMA_ENABLE_DEPRECATION_WARNINGS
+	}
 #endif
 }
 
diff --git a/Engine/Source/Runtime/IOS/MarketplaceKit/Private/MarketplaceKit.cpp b/Engine/Source/Runtime/IOS/MarketplaceKit/Private/MarketplaceKit.cpp
--- a/Engine/Source/Runtime/IOS/MarketplaceKit/Private/MarketplaceKit.cpp
+++ b/Engine/Source/Runtime/IOS/MarketplaceKit/Private/MarketplaceKit.cpp
@@ -6,6 +6,67 @@
 
 DEFINE_LOG_CATEGORY_STATIC(LogMarketplaceKit, Log, All);
 
+static FString GIOSTestFlightMarketplaceType = LexToString(EMarketplaceType::AppStore);
+static FAutoConsoleVariableRef CVarIOSTestFlightMarketplaceType(
+	TEXT("ios.testflightmarketplacetype"),
+	GIOSTestFlightMarketplaceType,
+	TEXT("Defines which marketplace type to return from MarketplaceKitModule, when running under testflight"),
+	ECVF_ReadOnly);
+
+static FString GIOSTestFlightMarketplaceBundleId;
+static FAutoConsoleVariableRef CVarIOSTestFlightMarketplaceMarketplaceBundleId(
+	TEXT("ios.testflightmarketplacebundleid"),
+	GIOSTestFlightMarketplaceBundleId,
+	TEXT("Defines which marketplace bundle id to return from MarketplaceKitModule, when ios.testflightmarketplacetype = Marketplace, and running under testflight"),
+	ECVF_ReadOnly);
+
+const TCHAR* LexToString(EMarketplaceType Value)
+{
+	switch (Value)
+	{
+		case EMarketplaceType::AppStore: return TEXT("AppStore");
+		case EMarketplaceType::TestFlight: return TEXT("TestFlight");
+		case EMarketplaceType::Marketplace: return TEXT("Marketplace");
+		case EMarketplaceType::Web: return TEXT("Web");
+		case EMarketplaceType::Other: return TEXT("Other");
+		default: checkNoEntry(); [[fallthrough]];
+		case EMarketplaceType::NotAvailable: return TEXT("NotAvailable");
+	}
+}
+
+void LexFromString(EMarketplaceType& OutValue, const TCHAR* InValue)
+{
+	if (FCString::Stricmp(InValue, TEXT("AppStore")) == 0)
+	{
+		OutValue = EMarketplaceType::AppStore;
+	}
+	else if (FCString::Stricmp(InValue, TEXT("TestFlight")) == 0)
+	{
+		OutValue = EMarketplaceType::TestFlight;
+	}
+	else if (FCString::Stricmp(InValue, TEXT("Marketplace")) == 0)
+	{
+		OutValue = EMarketplaceType::Marketplace;
+	}
+	else if (FCString::Stricmp(InValue, TEXT("Web")) == 0)
+	{
+		OutValue = EMarketplaceType::Web;
+	}
+	else if (FCString::Stricmp(InValue, TEXT("Other")) == 0)
+	{
+		OutValue = EMarketplaceType::Other;
+	}
+	else if (FCString::Stricmp(InValue, TEXT("NotAvailable")) == 0)
+	{
+		OutValue = EMarketplaceType::NotAvailable;
+	}
+	else
+	{
+		checkNoEntry();
+		OutValue = EMarketplaceType::NotAvailable;
+	}
+}
+
 void FMarketplaceKitModule::StartupModule()
 {
 	CacheValue();
@@ -38,20 +99,41 @@ void FMarketplaceKitModule::GetCurrentTypeAsync(TFunction<void(EMarketplaceType
 {
 	[AppDistributorWrapper getCurrentWithCompletionHandler:^(enum AppDistributorType Type, NSString* _Nonnull Name)
 	{
-		const EMarketplaceType ConvertedType = ConvertMarketplaceType(Type);
-		const FString ConvertedName(Name);
-
-		UE_LOG(LogMarketplaceKit, Log, TEXT("AppDistributorWrapper getCurrentWithCompletionHandler %i %s"), (int32)ConvertedType, *ConvertedName);
-
-		Callback(ConvertedType, ConvertedName);
+		{
+			const EMarketplaceType ConvertedType = ConvertMarketplaceType(Type);
+			const FString ConvertedName(Name);
+
+			UE_LOG(LogMarketplaceKit, Log, TEXT("AppDistributorWrapper getCurrentWithCompletionHandler %i %s"), (int32)ConvertedType, *ConvertedName);
+
+			CachedType = ConvertedType;
+			CachedName = ConvertedName;
+			bCachedTypeValid = true;
+		}
+
+		{
+			EMarketplaceType EffectiveType;
+			FString EffectiveName;
+			GetEffectiveType(EffectiveType, EffectiveName);
+			Callback(EffectiveType, EffectiveName);
+		}
 	}];
 }
 
 void FMarketplaceKitModule::GetCurrentType(EMarketplaceType& OutType, FString& OutName)
 {
 	CacheValue();
-	OutType = CachedType;
-	OutName = CachedName;
+	GetEffectiveType(OutType, OutName);
+}
+
+EMarketplaceType FMarketplaceKitModule::GetCurrentTypeStatic()
+{
+	EMarketplaceType Result = EMarketplaceType::NotAvailable;
+	if (FMarketplaceKitModule* MarketplaceKitModule = FModuleManager::LoadModulePtr<FMarketplaceKitModule>(TEXT("MarketplaceKit")))
+	{
+		FString Unused;
+		MarketplaceKitModule->GetCurrentType(Result, Unused);
+	}
+	return Result;
 }
 
 FString FMarketplaceKitModule::GetCurrentTypeAsString()
@@ -59,7 +141,12 @@ FString FMarketplaceKitModule::GetCurrentTypeAsString()
 	CacheValue();
 	
 	TStringBuilder<256> Result;
-	switch (CachedType)
+
+	EMarketplaceType EffectiveType;
+	FString EffectiveName;
+	GetEffectiveType(EffectiveType, EffectiveName);
+
+	switch (EffectiveType)
 	{
 		case EMarketplaceType::AppStore:		Result.Append(TEXT("AppStore")); break;
 		case EMarketplaceType::TestFlight:		Result.Append(TEXT("TestFlight")); break;
@@ -70,15 +157,30 @@ FString FMarketplaceKitModule::GetCurrentTypeAsString()
 		default:								Result.Append(TEXT("Other")); break;
 	}
 	
-	if (!CachedName.IsEmpty())
+	if (!EffectiveName.IsEmpty())
 	{
 		Result.Append(TEXT("-"));
-		Result.Append(CachedName);
+		Result.Append(MoveTemp(EffectiveName));
 	}
 
 	return *Result;
 }
 
+FString FMarketplaceKitModule::GetMarketplaceBundleIdStatic()
+{
+	if (FMarketplaceKitModule* MarketplaceKitModule = FModuleManager::LoadModulePtr<FMarketplaceKitModule>(TEXT("MarketplaceKit")))
+	{
+		EMarketplaceType Type;
+		FString Name;
+		MarketplaceKitModule->GetCurrentType(Type, Name);
+		if (Type == EMarketplaceType::Marketplace)
+		{
+			return Name;
+		}
+	}
+	return FString();
+}
+
 void FMarketplaceKitModule::CacheValue()
 {
 	if (bCachedTypeValid)
@@ -106,4 +208,18 @@ void FMarketplaceKitModule::CacheValue()
     dispatch_release(Semaphore);
 }
 
+void FMarketplaceKitModule::GetEffectiveType(EMarketplaceType& OutType, FString& OutName)
+{
+	if (CachedType == EMarketplaceType::TestFlight
+		// Other == run from XCode
+		|| CachedType == EMarketplaceType::Other)
+	{
+		LexFromString(OutType, *GIOSTestFlightMarketplaceType);
+		OutName = GIOSTestFlightMarketplaceBundleId;
+		return;
+	}
+	OutType = CachedType;
+	OutName = CachedName;
+}
+
 IMPLEMENT_MODULE(FMarketplaceKitModule, MarketplaceKit);
diff --git a/Engine/Source/Runtime/IOS/MarketplaceKit/Public/MarketplaceKit.h b/Engine/Source/Runtime/IOS/MarketplaceKit/Public/MarketplaceKit.h
--- a/Engine/Source/Runtime/IOS/MarketplaceKit/Public/MarketplaceKit.h
+++ b/Engine/Source/Runtime/IOS/MarketplaceKit/Public/MarketplaceKit.h
@@ -3,6 +3,7 @@
 #pragma once
 
 #include "CoreMinimal.h"
+#include "HAL/IConsoleManager.h"
 #include "Modules/ModuleManager.h"
 
 enum class EMarketplaceType : int32
@@ -15,6 +16,9 @@ enum class EMarketplaceType : int32
 	NotAvailable = 5,
 };
 
+const TCHAR* LexToString(EMarketplaceType Value);
+void LexFromString(EMarketplaceType& OutValue, const TCHAR* InValue);
+
 class FMarketplaceKitModule : public IModuleInterface
 {
 public:
@@ -23,13 +27,27 @@ public:
 	virtual bool SupportsDynamicReloading() override;
 
 	void GetCurrentTypeAsync(TFunction<void(EMarketplaceType Type, const FString& Name)> Callback);
+
+	/* Returns the MarketplaceType, and BundleId */
 	void GetCurrentType(EMarketplaceType& OutType, FString& OutName);
+	/* Returns the MarketplaceType */
+	static EMarketplaceType GetCurrentTypeStatic();
+	/* Returns the BundleId of the Marketplace, when MarketplaceType == Marketplace */
+	static FString GetMarketplaceBundleIdStatic();
+	
 	FString GetCurrentTypeAsString();
 
+	/* When installed from TestFlight, GetCurrentType etc above can return an emulated value for EMarketplaceType. If you _really_ need to know if you're running from TestFlight, this will tell you. */
+	bool IsTestFlight();
+	static bool IsTestFlightStatic();
+
 private:
 	bool bCachedTypeValid = false;
 	EMarketplaceType CachedType = EMarketplaceType::NotAvailable;
 	FString CachedName;
 
 	void CacheValue();
+	
+	// Redirects TestFlight to an emulated Type/Name set by CVar for testing.
+	void GetEffectiveType(EMarketplaceType& OutType, FString& OutName);
 };
