- Add GetCurrentTypeAsStringStatic static accessor

#rnx

diff --git a/Engine/Source/Runtime/IOS/MarketplaceKit/Private/MarketplaceKit.cpp b/Engine/Source/Runtime/IOS/MarketplaceKit/Private/MarketplaceKit.cpp
--- a/Engine/Source/Runtime/IOS/MarketplaceKit/Private/MarketplaceKit.cpp
+++ b/Engine/Source/Runtime/IOS/MarketplaceKit/Private/MarketplaceKit.cpp
@@ -166,6 +166,15 @@ FString FMarketplaceKitModule::GetCurrentTypeAsString()
 	return *Result;
 }
 
+FString FMarketplaceKitModule::GetCurrentTypeAsStringStatic()
+{
+	if (FMarketplaceKitModule* MarketplaceKitModule = FModuleManager::LoadModulePtr<FMarketplaceKitModule>(TEXT("MarketplaceKit")))
+	{
+		return MarketplaceKitModule->GetCurrentTypeAsString();
+	}
+	return FString();
+}
+
 FString FMarketplaceKitModule::GetMarketplaceBundleIdStatic()
 {
 	if (FMarketplaceKitModule* MarketplaceKitModule = FModuleManager::LoadModulePtr<FMarketplaceKitModule>(TEXT("MarketplaceKit")))
diff --git a/Engine/Source/Runtime/IOS/MarketplaceKit/Public/MarketplaceKit.h b/Engine/Source/Runtime/IOS/MarketplaceKit/Public/MarketplaceKit.h
--- a/Engine/Source/Runtime/IOS/MarketplaceKit/Public/MarketplaceKit.h
+++ b/Engine/Source/Runtime/IOS/MarketplaceKit/Public/MarketplaceKit.h
@@ -36,6 +36,7 @@ public:
 	static FString GetMarketplaceBundleIdStatic();
 	
 	FString GetCurrentTypeAsString();
+	static FString GetCurrentTypeAsStringStatic();
 
 	/* When installed from TestFlight, GetCurrentType etc above can return an emulated value for EMarketplaceType. If you _really_ need to know if you're running from TestFlight, this will tell you. */
 	bool IsTestFlight();
