- Change default when running from TestFlight / Other from "AppStore" (act like an AppStore app) to "TestFlight"

#rnx

diff --git a/Engine/Source/Runtime/IOS/MarketplaceKit/Private/MarketplaceKit.cpp b/Engine/Source/Runtime/IOS/MarketplaceKit/Private/MarketplaceKit.cpp
--- a/Engine/Source/Runtime/IOS/MarketplaceKit/Private/MarketplaceKit.cpp
+++ b/Engine/Source/Runtime/IOS/MarketplaceKit/Private/MarketplaceKit.cpp
@@ -6,7 +6,7 @@
 
 DEFINE_LOG_CATEGORY_STATIC(LogMarketplaceKit, Log, All);
 
-static FString GIOSTestFlightMarketplaceType = LexToString(EMarketplaceType::AppStore);
+static FString GIOSTestFlightMarketplaceType = LexToString(EMarketplaceType::TestFlight);
 static FAutoConsoleVariableRef CVarIOSTestFlightMarketplaceType(
 	TEXT("ios.testflightmarketplacetype"),
 	GIOSTestFlightMarketplaceType,
