- Fix using CVar in MarketplaceKit which is not initialized until too late for some callers

#rnx

diff --git a/Engine/Source/Runtime/IOS/MarketplaceKit/Private/MarketplaceKit.cpp b/Engine/Source/Runtime/IOS/MarketplaceKit/Private/MarketplaceKit.cpp
--- a/Engine/Source/Runtime/IOS/MarketplaceKit/Private/MarketplaceKit.cpp
+++ b/Engine/Source/Runtime/IOS/MarketplaceKit/Private/MarketplaceKit.cpp
@@ -6,20 +6,6 @@
 
 DEFINE_LOG_CATEGORY_STATIC(LogMarketplaceKit, Log, All);
 
-static FString GIOSTestFlightMarketplaceType = LexToString(EMarketplaceType::TestFlight);
-static FAutoConsoleVariableRef CVarIOSTestFlightMarketplaceType(
-	TEXT("ios.testflightmarketplacetype"),
-	GIOSTestFlightMarketplaceType,
-	TEXT("Defines which marketplace type to return from MarketplaceKitModule, when running under testflight"),
-	ECVF_ReadOnly);
-
-static FString GIOSTestFlightMarketplaceBundleId;
-static FAutoConsoleVariableRef CVarIOSTestFlightMarketplaceMarketplaceBundleId(
-	TEXT("ios.testflightmarketplacebundleid"),
-	GIOSTestFlightMarketplaceBundleId,
-	TEXT("Defines which marketplace bundle id to return from MarketplaceKitModule, when ios.testflightmarketplacetype = Marketplace, and running under testflight"),
-	ECVF_ReadOnly);
-
 const TCHAR* LexToString(EMarketplaceType Value)
 {
 	switch (Value)
@@ -69,6 +55,14 @@ void LexFromString(EMarketplaceType& OutValue, const TCHAR* InValue)
 
 void FMarketplaceKitModule::StartupModule()
 {
+	FString TestFlightMarketplaceTypeString;
+	if (GConfig->GetString(TEXT("MarketplaceKit"), TEXT("TestFlightMarketplaceType"), TestFlightMarketplaceTypeString, GEngineIni))
+	{
+		LexFromString(TestFlightMarketplaceType, *TestFlightMarketplaceTypeString);
+	}
+
+	GConfig->GetString(TEXT("MarketplaceKit"), TEXT("TestFlightMarketplaceBundleId"), TestFlightMarketplaceBundleId, GEngineIni);
+
 	CacheValue();
 }
 
@@ -223,8 +217,8 @@ void FMarketplaceKitModule::GetEffectiveType(EMarketplaceType& OutType, FString&
 		// Other == run from XCode
 		|| CachedType == EMarketplaceType::Other)
 	{
-		LexFromString(OutType, *GIOSTestFlightMarketplaceType);
-		OutName = GIOSTestFlightMarketplaceBundleId;
+		OutType = TestFlightMarketplaceType;
+		OutName = TestFlightMarketplaceBundleId;
 		return;
 	}
 	OutType = CachedType;
diff --git a/Engine/Source/Runtime/IOS/MarketplaceKit/Public/MarketplaceKit.h b/Engine/Source/Runtime/IOS/MarketplaceKit/Public/MarketplaceKit.h
--- a/Engine/Source/Runtime/IOS/MarketplaceKit/Public/MarketplaceKit.h
+++ b/Engine/Source/Runtime/IOS/MarketplaceKit/Public/MarketplaceKit.h
@@ -47,6 +47,9 @@ private:
 	EMarketplaceType CachedType = EMarketplaceType::NotAvailable;
 	FString CachedName;
 
+	EMarketplaceType TestFlightMarketplaceType = EMarketplaceType::TestFlight;
+	FString TestFlightMarketplaceBundleId;
+
 	void CacheValue();
 	
 	// Redirects TestFlight to an emulated Type/Name set by CVar for testing.
