- Map NotAvailable to AppStore on the basis this should only be seen on pre 17.4 IOS devices.
- Improve logging to use the string not the number.

#rnx

diff --git a/Engine/Source/Runtime/IOS/MarketplaceKit/Private/MarketplaceKit.cpp b/Engine/Source/Runtime/IOS/MarketplaceKit/Private/MarketplaceKit.cpp
--- a/Engine/Source/Runtime/IOS/MarketplaceKit/Private/MarketplaceKit.cpp
+++ b/Engine/Source/Runtime/IOS/MarketplaceKit/Private/MarketplaceKit.cpp
@@ -84,7 +84,7 @@ static constexpr EMarketplaceType ConvertMarketplaceType(const AppDistributorTyp
 		case AppDistributorTypeMarketplace:		return EMarketplaceType::Marketplace;
 		case AppDistributorTypeWeb:				return EMarketplaceType::Web;
 		case AppDistributorTypeOther:			return EMarketplaceType::Other;
-		case AppDistributorTypeNotAvailable:	return EMarketplaceType::NotAvailable;
+		case AppDistributorTypeNotAvailable:	return EMarketplaceType::AppStore; // Pre 17.4, hardcode to AppStore
 		default:								return EMarketplaceType::Other;
 	}
 }
@@ -97,7 +97,7 @@ void FMarketplaceKitModule::GetCurrentTypeAsync(TFunction<void(EMarketplaceType
 			const EMarketplaceType ConvertedType = ConvertMarketplaceType(Type);
 			const FString ConvertedName(Name);
 
-			UE_LOG(LogMarketplaceKit, Log, TEXT("AppDistributorWrapper getCurrentWithCompletionHandler %i %s"), (int32)ConvertedType, *ConvertedName);
+			UE_LOG(LogMarketplaceKit, Log, TEXT("AppDistributorWrapper getCurrentWithCompletionHandler %s %s"), LexToString(ConvertedType), *ConvertedName);
 
 			CachedType = ConvertedType;
 			CachedName = ConvertedName;
@@ -201,7 +201,7 @@ void FMarketplaceKitModule::CacheValue()
 		CachedName = Name;
 		bCachedTypeValid = true;
 
-		UE_LOG(LogMarketplaceKit, Log, TEXT("AppDistributorWrapper getCurrentWithCompletionHandler %i %s"), (int32)CachedType, *CachedName);
+		UE_LOG(LogMarketplaceKit, Log, TEXT("AppDistributorWrapper getCurrentWithCompletionHandler %s %s"), LexToString(CachedType), *CachedName);
 
 		dispatch_semaphore_signal(Semaphore);
 	}];
@@ -217,7 +217,7 @@ void FMarketplaceKitModule::GetEffectiveType(EMarketplaceType& OutType, FString&
 		// Other == run from XCode
 		|| CachedType == EMarketplaceType::Other)
 	{
-		OutType = TestFlightMarketplaceType;
+		OutType = TestFlightMarketplaceType);
 		OutName = TestFlightMarketplaceBundleId;
 		return;
 	}
