- Fix CIS

#rnx

diff --git a/Engine/Source/Runtime/IOS/MarketplaceKit/Private/MarketplaceKit.cpp b/Engine/Source/Runtime/IOS/MarketplaceKit/Private/MarketplaceKit.cpp
--- a/Engine/Source/Runtime/IOS/MarketplaceKit/Private/MarketplaceKit.cpp
+++ b/Engine/Source/Runtime/IOS/MarketplaceKit/Private/MarketplaceKit.cpp
@@ -217,7 +217,7 @@ void FMarketplaceKitModule::GetEffectiveType(EMarketplaceType& OutType, FString&
 		// Other == run from XCode
 		|| CachedType == EMarketplaceType::Other)
 	{
-		OutType = TestFlightMarketplaceType);
+		OutType = TestFlightMarketplaceType;
 		OutName = TestFlightMarketplaceBundleId;
 		return;
 	}
