MpKit: Make semaphore wait for forever.

This resolves an edge case where the system is too busy and we end up waiting a more than a second for the swift code to return.

Bad things happen if we don't get the proper value from this.

diff --git a/Engine/Source/Runtime/IOS/MarketplaceKit/Private/MarketplaceKit.cpp b/Engine/Source/Runtime/IOS/MarketplaceKit/Private/MarketplaceKit.cpp
--- a/Engine/Source/Runtime/IOS/MarketplaceKit/Private/MarketplaceKit.cpp
+++ b/Engine/Source/Runtime/IOS/MarketplaceKit/Private/MarketplaceKit.cpp
@@ -207,7 +207,7 @@ void FMarketplaceKitModule::CacheValue()
 	}];
 
 	// wait for a result, but timeout after 1s
-	dispatch_semaphore_wait(Semaphore, dispatch_time(DISPATCH_TIME_NOW, NSEC_PER_SEC));
+	dispatch_semaphore_wait(Semaphore, DISPATCH_TIME_FOREVER);
     dispatch_release(Semaphore);
 }
 
