- Fix non-unity CIS

#rnx

diff --git a/Engine/Source/Runtime/IOS/MarketplaceKit/Private/MarketplaceKit.cpp b/Engine/Source/Runtime/IOS/MarketplaceKit/Private/MarketplaceKit.cpp
--- a/Engine/Source/Runtime/IOS/MarketplaceKit/Private/MarketplaceKit.cpp
+++ b/Engine/Source/Runtime/IOS/MarketplaceKit/Private/MarketplaceKit.cpp
@@ -1,7 +1,10 @@
 // Copyright Epic Games, Inc. All Rights Reserved.
 
 #include "MarketplaceKit.h"
+
+#include "CoreGlobals.h"
 #include "MarketplaceKitWrapper.h"
+#include "Misc/ConfigCacheIni.h"
 #include "Logging/LogMacros.h"
 
 DEFINE_LOG_CATEGORY_STATIC(LogMarketplaceKit, Log, All);
