- Add missing IsTestFlight / IsTestFlightStatic implementations

#rnx

diff --git a/Engine/Source/Runtime/IOS/MarketplaceKit/Private/MarketplaceKit.cpp b/Engine/Source/Runtime/IOS/MarketplaceKit/Private/MarketplaceKit.cpp
--- a/Engine/Source/Runtime/IOS/MarketplaceKit/Private/MarketplaceKit.cpp
+++ b/Engine/Source/Runtime/IOS/MarketplaceKit/Private/MarketplaceKit.cpp
@@ -187,6 +187,21 @@ FString FMarketplaceKitModule::GetMarketplaceBundleIdStatic()
 	return FString();
 }
 
+bool FMarketplaceKitModule::IsTestFlight() const
+{
+	CacheValue();
+	return CachedType == EMarketplaceType::TestFlight;
+}
+
+bool FMarketplaceKitModule::IsTestFlightStatic()
+{
+	if (FMarketplaceKitModule* MarketplaceKitModule = FModuleManager::LoadModulePtr<FMarketplaceKitModule>(TEXT("MarketplaceKit")))
+	{
+		return MarketplaceKitModule->IsTestFlight();
+	}
+	return false;
+}
+
 void FMarketplaceKitModule::CacheValue()
 {
 	if (bCachedTypeValid)
diff --git a/Engine/Source/Runtime/IOS/MarketplaceKit/Public/MarketplaceKit.h b/Engine/Source/Runtime/IOS/MarketplaceKit/Public/MarketplaceKit.h
--- a/Engine/Source/Runtime/IOS/MarketplaceKit/Public/MarketplaceKit.h
+++ b/Engine/Source/Runtime/IOS/MarketplaceKit/Public/MarketplaceKit.h
@@ -39,7 +39,7 @@ public:
 	static FString GetCurrentTypeAsStringStatic();
 
 	/* When installed from TestFlight, GetCurrentType etc above can return an emulated value for EMarketplaceType. If you _really_ need to know if you're running from TestFlight, this will tell you. */
-	bool IsTestFlight();
+	bool IsTestFlight() const;
 	static bool IsTestFlightStatic();
 
 private:
