- Fix CIS

#rnx

diff --git a/Engine/Source/Runtime/IOS/MarketplaceKit/Private/MarketplaceKit.cpp b/Engine/Source/Runtime/IOS/MarketplaceKit/Private/MarketplaceKit.cpp
--- a/Engine/Source/Runtime/IOS/MarketplaceKit/Private/MarketplaceKit.cpp
+++ b/Engine/Source/Runtime/IOS/MarketplaceKit/Private/MarketplaceKit.cpp
@@ -187,7 +187,7 @@ FString FMarketplaceKitModule::GetMarketplaceBundleIdStatic()
 	return FString();
 }
 
-bool FMarketplaceKitModule::IsTestFlight() const
+bool FMarketplaceKitModule::IsTestFlight()
 {
 	CacheValue();
 	return CachedType == EMarketplaceType::TestFlight;
diff --git a/Engine/Source/Runtime/IOS/MarketplaceKit/Public/MarketplaceKit.h b/Engine/Source/Runtime/IOS/MarketplaceKit/Public/MarketplaceKit.h
--- a/Engine/Source/Runtime/IOS/MarketplaceKit/Public/MarketplaceKit.h
+++ b/Engine/Source/Runtime/IOS/MarketplaceKit/Public/MarketplaceKit.h
@@ -39,7 +39,7 @@ public:
 	static FString GetCurrentTypeAsStringStatic();
 
 	/* When installed from TestFlight, GetCurrentType etc above can return an emulated value for EMarketplaceType. If you _really_ need to know if you're running from TestFlight, this will tell you. */
-	bool IsTestFlight() const;
+	bool IsTestFlight();
 	static bool IsTestFlightStatic();
 
 private:
