- Add commandline switches for testflight marketplace type/bundleid

#rnx

diff --git a/Engine/Source/Runtime/IOS/MarketplaceKit/Private/MarketplaceKit.cpp b/Engine/Source/Runtime/IOS/MarketplaceKit/Private/MarketplaceKit.cpp
--- a/Engine/Source/Runtime/IOS/MarketplaceKit/Private/MarketplaceKit.cpp
+++ b/Engine/Source/Runtime/IOS/MarketplaceKit/Private/MarketplaceKit.cpp
@@ -4,7 +4,9 @@
 
 #include "CoreGlobals.h"
 #include "MarketplaceKitWrapper.h"
+#include "Misc/CommandLine.h"
 #include "Misc/ConfigCacheIni.h"
+#include "Misc/Parse.h"
 #include "Logging/LogMacros.h"
 
 DEFINE_LOG_CATEGORY_STATIC(LogMarketplaceKit, Log, All);
@@ -59,12 +61,15 @@ void LexFromString(EMarketplaceType& OutValue, const TCHAR* InValue)
 void FMarketplaceKitModule::StartupModule()
 {
 	FString TestFlightMarketplaceTypeString;
-	if (GConfig->GetString(TEXT("MarketplaceKit"), TEXT("TestFlightMarketplaceType"), TestFlightMarketplaceTypeString, GEngineIni))
+	GConfig->GetString(TEXT("MarketplaceKit"), TEXT("TestFlightMarketplaceType"), TestFlightMarketplaceTypeString, GEngineIni);
+	FParse::Value(FCommandLine::Get(), TEXT("TestFlightMarketplaceType="), TestFlightMarketplaceTypeString);
+	if (!TestFlightMarketplaceTypeString.IsEmpty())
 	{
 		LexFromString(TestFlightMarketplaceType, *TestFlightMarketplaceTypeString);
 	}
 
 	GConfig->GetString(TEXT("MarketplaceKit"), TEXT("TestFlightMarketplaceBundleId"), TestFlightMarketplaceBundleId, GEngineIni);
+	FParse::Value(FCommandLine::Get(), TEXT("TestFlightMarketplaceBundleId="), TestFlightMarketplaceBundleId);
 
 	CacheValue();
 }
