- Ensure we only include the bundle id when the Type is Marketplace
-- we are seeing a few instances of AppStore-AppStore in the analytics

#rnx

diff --git a/Engine/Source/Runtime/IOS/MarketplaceKit/Private/MarketplaceKit.cpp b/Engine/Source/Runtime/IOS/MarketplaceKit/Private/MarketplaceKit.cpp
--- a/Engine/Source/Runtime/IOS/MarketplaceKit/Private/MarketplaceKit.cpp
+++ b/Engine/Source/Runtime/IOS/MarketplaceKit/Private/MarketplaceKit.cpp
@@ -159,7 +159,7 @@ FString FMarketplaceKitModule::GetCurrentTypeAsString()
 		default:								Result.Append(TEXT("Other")); break;
 	}
 	
-	if (!EffectiveName.IsEmpty())
+	if (EffectiveType == EMarketplaceType::Marketplace && !EffectiveName.IsEmpty())
 	{
 		Result.Append(TEXT("-"));
 		Result.Append(MoveTemp(EffectiveName));
