Switch to using an embedded weak framework for the MarketplaceKit swift lib to allow supporting APIs not yet available on device.

#Synced-CL 52493164

#Synced-CL 52493164

diff --git a/Engine/Plugins/Online/IOS/OnlineSubsystemIOS/Source/Private/OnlineStoreIOS.cpp b/Engine/Plugins/Online/IOS/OnlineSubsystemIOS/Source/Private/OnlineStoreIOS.cpp
--- a/Engine/Plugins/Online/IOS/OnlineSubsystemIOS/Source/Private/OnlineStoreIOS.cpp
+++ b/Engine/Plugins/Online/IOS/OnlineSubsystemIOS/Source/Private/OnlineStoreIOS.cpp
@@ -17,6 +17,10 @@
 #include "IOSAppDelegate.h"
 #include "IOS/IOSAsyncTask.h"
 
+#if !PLATFORM_TVOS && !PLATFORM_VISIONOS
+#include "MarketplaceKit.h"
+#endif
+
 @interface FSKProductsRequestHelper : SKProductsRequest
 {
 };
@@ -335,6 +339,34 @@ void FOnlineStoreIOS::QueryStoreCountry()
 	}];
 }
 
+void FOnlineStoreIOS::QueryEligibilityRegion()
+{
+#if !PLATFORM_TVOS && !PLATFORM_VISIONOS
+	TWeakPtr<FOnlineStoreIOS> WeakThis = AsWeak();
+																													  
+	TDelegate<void(bool, FString)> Callback = TDelegate<void(bool, FString)>::CreateLambda(
+		[WeakThis](bool Success, FString Output)                                                                      
+		{                                                                                                             
+			if (TSharedPtr<FOnlineStoreIOS> SharedThis = WeakThis.Pin())                                          
+			{
+				if (Success)
+				{
+					SharedThis->EligibilityRegion = Output;                                                       
+				}                                                                                                 
+				else
+				{
+					SharedThis->EligibilityRegion = "";                                                       
+					UE_LOG_ONLINE_STOREV2(Warning, TEXT("Unable to obtain eligibility region: %s"), *Output);
+				}
+				SharedThis->bIsEligibilityRegionInitialized.Store(true);
+			}
+		}
+	);
+	
+	FMarketplaceKitModule::GetEligibilityRegionAsync(Callback);
+#endif // !PLATFORM_TVOS && !PLATFORM_VISIONOS
+}
+
 FString FOnlineStoreIOS::GetStoreCountry() const
 {
 	if (bIsStoreCountryInitialized.Load())
@@ -346,6 +378,17 @@ FString FOnlineStoreIOS::GetStoreCountry() const
 	return "";
 }
 
+FString FOnlineStoreIOS::GetEligibilityRegion() const
+{
+	if (bIsEligibilityRegionInitialized.Load())
+	{
+		return EligibilityRegion;
+	}
+	
+	UE_LOG_ONLINE_STOREV2(Warning, TEXT("Attempting to get eligibility region before it is initialized!"));
+	return "";
+}
+
 #if !UE_BUILD_SHIPPING
 class FStoreExec : public FSelfRegisteringExec
 {
diff --git a/Engine/Plugins/Online/IOS/OnlineSubsystemIOS/Source/Private/OnlineSubsystemIOS.cpp b/Engine/Plugins/Online/IOS/OnlineSubsystemIOS/Source/Private/OnlineSubsystemIOS.cpp
--- a/Engine/Plugins/Online/IOS/OnlineSubsystemIOS/Source/Private/OnlineSubsystemIOS.cpp
+++ b/Engine/Plugins/Online/IOS/OnlineSubsystemIOS/Source/Private/OnlineSubsystemIOS.cpp
@@ -208,6 +208,9 @@ bool FOnlineSubsystemIOS::Init()
 		{
 			// Initialize the store country
 			StoreV2Interface->QueryStoreCountry();
+			
+			// Initialize the eligibility region
+			StoreV2Interface->QueryEligibilityRegion();
 		}
 	}
 
diff --git a/Engine/Plugins/Online/IOS/OnlineSubsystemIOS/Source/Public/OnlineStoreIOS.h b/Engine/Plugins/Online/IOS/OnlineSubsystemIOS/Source/Public/OnlineStoreIOS.h
--- a/Engine/Plugins/Online/IOS/OnlineSubsystemIOS/Source/Public/OnlineStoreIOS.h
+++ b/Engine/Plugins/Online/IOS/OnlineSubsystemIOS/Source/Public/OnlineStoreIOS.h
@@ -31,6 +31,7 @@ public:
 	virtual void GetOffers(TArray<FOnlineStoreOfferRef>& OutOffers) const override;
 	virtual TSharedPtr<FOnlineStoreOffer> GetOffer(const FUniqueOfferId& OfferId) const override;
 	virtual FString GetStoreCountry() const override;
+	virtual FString GetEligibilityRegion() const override;
 
 	// FOnlineStoreIOS
 
@@ -63,6 +64,11 @@ public:
 	 * Starts an async request to get the store kit country
 	 */
 	void QueryStoreCountry();
+	
+	/**
+	 * Starts an async request to get the marketplace kit eligibility region
+	 */
+	void QueryEligibilityRegion();
 
 private:
 	
@@ -115,6 +121,9 @@ private:
 	FString StoreCountry;
 	TAtomic<bool> bIsStoreCountryInitialized = false;
 	
+	FString EligibilityRegion;
+	TAtomic<bool> bIsEligibilityRegionInitialized = false;
+	
 	/** Reference to the parent subsystem */
 	FOnlineSubsystemIOS* Subsystem;
 };
diff --git a/Engine/Plugins/Online/OnlineSubsystem/Source/Public/Interfaces/OnlineStoreInterfaceV2.h b/Engine/Plugins/Online/OnlineSubsystem/Source/Public/Interfaces/OnlineStoreInterfaceV2.h
--- a/Engine/Plugins/Online/OnlineSubsystem/Source/Public/Interfaces/OnlineStoreInterfaceV2.h
+++ b/Engine/Plugins/Online/OnlineSubsystem/Source/Public/Interfaces/OnlineStoreInterfaceV2.h
@@ -277,5 +277,9 @@ public:
 	* Return store country if it has been fetched.
 	*/
 	virtual FString GetStoreCountry() const { return ""; }
-
+	
+	/**
+	* Return eligibility region if it has been fetched.
+	*/
+	virtual FString GetEligibilityRegion() const { return ""; }
 };
diff --git a/Engine/Source/Runtime/IOS/MarketplaceKit/MarketplaceKit.Build.cs b/Engine/Source/Runtime/IOS/MarketplaceKit/MarketplaceKit.Build.cs
--- a/Engine/Source/Runtime/IOS/MarketplaceKit/MarketplaceKit.Build.cs
+++ b/Engine/Source/Runtime/IOS/MarketplaceKit/MarketplaceKit.Build.cs
@@ -2,6 +2,7 @@
 
 using UnrealBuildTool;
 using System.IO;
+using System;
 
 public class MarketplaceKit : ModuleRules
 {
@@ -9,14 +10,19 @@ public class MarketplaceKit : ModuleRules
 	{
 		bRequiresPlatformSDK = true;
 
-		PrivateDependencyModuleNames.Add("Swift");
-
 		PrivateDependencyModuleNames.AddRange(new string[]
 		{
 			"Core",
 			"Engine",
 		});
 
-		PublicAdditionalLibraries.Add(Path.Combine(ModuleDirectory, "arm64", "libMarketplaceKitWrapper.a"));
+		// Embed the pre-built MarketplaceKitWrapper framework into the app bundle
+		PublicAdditionalFrameworks.Add(new Framework(
+			"MarketplaceKitWrapper",
+			Path.Combine(ModuleDirectory, "MarketplaceKitWrapper.embeddedframework.zip"),
+			Framework.FrameworkMode.Copy));
+
+		// Weak-link so the app launches even if the framework is unavailable at runtime
+		PublicWeakFrameworks.Add("MarketplaceKitWrapper");
 	}
 }
diff --git a/Engine/Source/Runtime/IOS/MarketplaceKit/Private/MarketplaceKit.cpp b/Engine/Source/Runtime/IOS/MarketplaceKit/Private/MarketplaceKit.cpp
--- a/Engine/Source/Runtime/IOS/MarketplaceKit/Private/MarketplaceKit.cpp
+++ b/Engine/Source/Runtime/IOS/MarketplaceKit/Private/MarketplaceKit.cpp
@@ -3,25 +3,33 @@
 #include "MarketplaceKit.h"
 
 #include "CoreGlobals.h"
-#include "MarketplaceKitWrapper.h"
+#include "MarketplaceKitWrapperAPI.h"
 #include "Misc/CommandLine.h"
 #include "Misc/ConfigCacheIni.h"
 #include "Misc/Parse.h"
 #include "Logging/LogMacros.h"
 
+#include <dispatch/dispatch.h>
+
 DEFINE_LOG_CATEGORY_STATIC(LogMarketplaceKit, Log, All);
 
-const TCHAR* LexToString(EMarketplaceType Value)
+static bool IsMarketplaceKitWrapperAvailable()
+{
+	// Weak-linked symbols are nullptr if the framework isn't loaded
+	return MKW_GetCurrent != nullptr;
+}
+
+const FString LexToString(EMarketplaceType Value)
 {
 	switch (Value)
 	{
-		case EMarketplaceType::AppStore: return TEXT("AppStore");
-		case EMarketplaceType::TestFlight: return TEXT("TestFlight");
-		case EMarketplaceType::Marketplace: return TEXT("Marketplace");
-		case EMarketplaceType::Web: return TEXT("Web");
-		case EMarketplaceType::Other: return TEXT("Other");
+		case EMarketplaceType::AppStore: return "AppStore";
+		case EMarketplaceType::TestFlight: return "TestFlight";
+		case EMarketplaceType::Marketplace: return "Marketplace";
+		case EMarketplaceType::Web: return "Web";
+		case EMarketplaceType::Other: return "Other";
 		default: checkNoEntry(); [[fallthrough]];
-		case EMarketplaceType::NotAvailable: return TEXT("NotAvailable");
+		case EMarketplaceType::NotAvailable: return "NotAvailable";
 	}
 }
 
@@ -71,6 +79,11 @@ void FMarketplaceKitModule::StartupModule()
 	GConfig->GetString(TEXT("MarketplaceKit"), TEXT("TestFlightMarketplaceBundleId"), TestFlightMarketplaceBundleId, GEngineIni);
 	FParse::Value(FCommandLine::Get(), TEXT("TestFlightMarketplaceBundleId="), TestFlightMarketplaceBundleId);
 
+	if (!IsMarketplaceKitWrapperAvailable())
+	{
+		UE_LOGF(LogMarketplaceKit, Warning, "MarketplaceKitWrapper framework not available, marketplace features will be disabled");
+	}
+
 	CacheValue();
 }
 
@@ -83,50 +96,106 @@ bool FMarketplaceKitModule::SupportsDynamicReloading()
 	return true;
 }
 
-static constexpr EMarketplaceType ConvertMarketplaceType(const AppDistributorType Type)
+static constexpr EMarketplaceType ConvertMarketplaceType(const long Type)
 {
 	switch (Type)
 	{
-		case AppDistributorTypeAppStore:		return EMarketplaceType::AppStore;
-		case AppDistributorTypeTestFlight:		return EMarketplaceType::TestFlight;
-		case AppDistributorTypeMarketplace:		return EMarketplaceType::Marketplace;
-		case AppDistributorTypeWeb:				return EMarketplaceType::Web;
-		case AppDistributorTypeOther:			return EMarketplaceType::Other;
-		case AppDistributorTypeNotAvailable:	return EMarketplaceType::AppStore; // Pre 17.4, hardcode to AppStore
+		case 0:		return EMarketplaceType::AppStore;
+		case 1:		return EMarketplaceType::TestFlight;
+		case 2:		return EMarketplaceType::Marketplace;
+		case 3:		return EMarketplaceType::Web;
+		case 4:		return EMarketplaceType::Other;
+		case 5:		return EMarketplaceType::AppStore; // NotAvailable: Pre 17.4, hardcode to AppStore
 		default:								return EMarketplaceType::Other;
 	}
 }
 
 void FMarketplaceKitModule::GetCurrentTypeAsync(TFunction<void(EMarketplaceType Type, const FString& Name)> Callback)
 {
-	[AppDistributorWrapper getCurrentWithCompletionHandler:^(enum AppDistributorType Type, NSString* _Nonnull Name)
+	if (!IsMarketplaceKitWrapperAvailable())
 	{
+		Callback(EMarketplaceType::NotAvailable, FString());
+		return;
+	}
+
+	// C callbacks don't support captures, so store context in statics.
+	// The Swift side dispatches the callback on the main thread, so only one can be in flight at a time.
+	static TFunction<void(EMarketplaceType, const FString&)>* PendingGetCurrentCallback = nullptr;
+	static FMarketplaceKitModule* PendingSelf = nullptr;
+
+	PendingGetCurrentCallback = new TFunction<void(EMarketplaceType, const FString&)>(MoveTemp(Callback));
+	PendingSelf = this;
+
+	MKW_GetCurrent([](long Type, const char* Name)
 		{
 			const EMarketplaceType ConvertedType = ConvertMarketplaceType(Type);
-			const FString ConvertedName(Name);
+		const FString ConvertedName = FString(UTF8_TO_TCHAR(Name ? Name : ""));
 
-			UE_LOG(LogMarketplaceKit, Log, TEXT("AppDistributorWrapper getCurrentWithCompletionHandler %s %s"), LexToString(ConvertedType), *ConvertedName);
+		UE_LOGF(LogMarketplaceKit, Log, "MKW_GetCurrent callback %ls %ls", *LexToString(ConvertedType), *ConvertedName);
 
-			CachedType = ConvertedType;
-			CachedName = ConvertedName;
-			bCachedTypeValid = true;
+		if (PendingSelf)
+		{
+			PendingSelf->CachedType = ConvertedType;
+			PendingSelf->CachedName = ConvertedName;
+			PendingSelf->bCachedTypeValid = true;
 		}
 
+		if (PendingGetCurrentCallback)
 		{
 			EMarketplaceType EffectiveType;
 			FString EffectiveName;
-			GetEffectiveType(EffectiveType, EffectiveName);
-			Callback(EffectiveType, EffectiveName);
+			if (PendingSelf)
+			{
+				PendingSelf->GetEffectiveType(EffectiveType, EffectiveName);
+			}
+			else
+			{
+				EffectiveType = ConvertedType;
+				EffectiveName = ConvertedName;
+			}
+
+			(*PendingGetCurrentCallback)(EffectiveType, EffectiveName);
+			delete PendingGetCurrentCallback;
+			PendingGetCurrentCallback = nullptr;
+			PendingSelf = nullptr;
 		}
-	}];
+	});
 }
 
 void FMarketplaceKitModule::RequestCTTokenAsync(TDelegate<void(bool Success, FString Output)> Callback)
 {
-	[AppDistributorWrapper requestCTToken:^(BOOL Success, NSString * _Nonnull Output)
+	if (!MKW_RequestCTToken)
 	{
-		Callback.ExecuteIfBound(Success, Output);
-	}];
+		Callback.ExecuteIfBound(false, FString(TEXT("MarketplaceKitWrapper framework not loaded")));
+		return;
+	}
+
+	static TDelegate<void(bool, FString)> PendingCTTokenCallback;
+	PendingCTTokenCallback = MoveTemp(Callback);
+
+	MKW_RequestCTToken([](bool Success, const char* Output)
+	{
+		PendingCTTokenCallback.ExecuteIfBound(Success, FString(UTF8_TO_TCHAR(Output ? Output : "")));
+		PendingCTTokenCallback.Unbind();
+	});
+}
+
+void FMarketplaceKitModule::GetEligibilityRegionAsync(TDelegate<void(bool Success, FString Output)> Callback)
+{
+	if (!MKW_GetEligibilityRegion)
+	{
+		Callback.ExecuteIfBound(false, FString(TEXT("MarketplaceKitWrapper framework not loaded")));
+		return;
+	}
+
+	static TDelegate<void(bool, FString)> PendingEligibilityRegionCallback;
+	PendingEligibilityRegionCallback = MoveTemp(Callback);
+
+	MKW_GetEligibilityRegion([](bool Success, const char* Output)
+	{
+		PendingEligibilityRegionCallback.ExecuteIfBound(Success, FString(UTF8_TO_TCHAR(Output ? Output : "")));
+		PendingEligibilityRegionCallback.Unbind();
+	});
 }
 
 void FMarketplaceKitModule::GetCurrentType(EMarketplaceType& OutType, FString& OutName)
@@ -222,20 +291,35 @@ void FMarketplaceKitModule::CacheValue()
 		return;
 	}
 
+	if (!IsMarketplaceKitWrapperAvailable())
+	{
+		return;
+	}
+
 	// TODO avoid scheduling multiple requests in case this path is hit from multiple threads
 
 	dispatch_semaphore_t Semaphore = dispatch_semaphore_create(0);
 
-	[AppDistributorWrapper getCurrentWithCompletionHandler:^(enum AppDistributorType Type, NSString* _Nonnull Name)
+	static EMarketplaceType* PendingCacheType = nullptr;
+	static FString* PendingCacheName = nullptr;
+	static bool* PendingCacheValid = nullptr;
+	static dispatch_semaphore_t PendingCacheSemaphore = nullptr;
+
+	PendingCacheType = &CachedType;
+	PendingCacheName = &CachedName;
+	PendingCacheValid = &bCachedTypeValid;
+	PendingCacheSemaphore = Semaphore;
+
+	MKW_GetCurrent([](long Type, const char* Name)
 	{
-		CachedType = ConvertMarketplaceType(Type);
-		CachedName = Name;
-		bCachedTypeValid = true;
+		*PendingCacheType = ConvertMarketplaceType(Type);
+		*PendingCacheName = FString(UTF8_TO_TCHAR(Name ? Name : ""));
+		*PendingCacheValid = true;
 
-		UE_LOG(LogMarketplaceKit, Log, TEXT("AppDistributorWrapper getCurrentWithCompletionHandler %s %s"), LexToString(CachedType), *CachedName);
+		UE_LOGF(LogMarketplaceKit, Log, "MKW_GetCurrent CacheValue callback %ls %ls", *LexToString(*PendingCacheType), **PendingCacheName);
 
-		dispatch_semaphore_signal(Semaphore);
-	}];
+		dispatch_semaphore_signal(PendingCacheSemaphore);
+	});
 
 	// wait for a result, but timeout after 1s
 	dispatch_semaphore_wait(Semaphore, DISPATCH_TIME_FOREVER);
diff --git a/Engine/Source/Runtime/IOS/MarketplaceKit/Private/MarketplaceKitWrapper.h b/Engine/Source/Runtime/IOS/MarketplaceKit/Private/MarketplaceKitWrapper.h
--- a/Engine/Source/Runtime/IOS/MarketplaceKit/Private/MarketplaceKitWrapper.h
+++ b/Engine/Source/Runtime/IOS/MarketplaceKit/Private/MarketplaceKitWrapper.h
@@ -1,5 +1,5 @@
 // Copyright Epic Games, Inc. All Rights Reserved.
-// Generated by Apple Swift version 6.2.3 effective-5.10 (swiftlang-6.2.3.3.21 clang-1700.6.3.2)
+// Generated by Apple Swift version 6.3 effective-5.10 (swiftlang-6.3.0.123.5 clang-2100.0.123.102)
 #ifndef MARKETPLACEKITWRAPPER_SWIFT_H
 #define MARKETPLACEKITWRAPPER_SWIFT_H
 #pragma clang diagnostic push
@@ -25,7 +25,7 @@
 #pragma clang diagnostic ignored "-Wauto-import"
 #if defined(__OBJC__)
 #include <Foundation/Foundation.h>
-#endif
+#endif // defined(__OBJC__)
 #if defined(__cplusplus)
 #include <cstdint>
 #include <cstddef>
@@ -194,13 +194,32 @@ typedef unsigned int swift_uint4  __attribute__((__ext_vector_type__(4)));
 # endif
 #endif
 #if !defined(SWIFT_ENUM)
-# define SWIFT_ENUM(_type, _name, _extensibility) enum _name : _type _name; enum SWIFT_ENUM_ATTR(_extensibility) SWIFT_ENUM_EXTRA _name : _type
-# if __has_feature(generalized_swift_name)
-#  define SWIFT_ENUM_NAMED(_type, _name, SWIFT_NAME, _extensibility) enum _name : _type _name SWIFT_COMPILE_NAME(SWIFT_NAME); enum SWIFT_COMPILE_NAME(SWIFT_NAME) SWIFT_ENUM_ATTR(_extensibility) SWIFT_ENUM_EXTRA _name : _type
+# if (defined(__cplusplus) && __cplusplus >= 201103L) ||      (defined(__STDC_VERSION__) && __STDC_VERSION__ >= 202311L) ||      __has_feature(objc_fixed_enum)
+#  define SWIFT_ENUM(_type, _name, _extensibility) enum _name : _type _name; enum SWIFT_ENUM_ATTR(_extensibility) SWIFT_ENUM_EXTRA _name : _type
+#  if __has_feature(generalized_swift_name)
+#   define SWIFT_ENUM_NAMED(_type, _name, SWIFT_NAME, _extensibility) enum _name : _type _name SWIFT_COMPILE_NAME(SWIFT_NAME); enum SWIFT_COMPILE_NAME(SWIFT_NAME) SWIFT_ENUM_ATTR(_extensibility) SWIFT_ENUM_EXTRA _name : _type
+#  else
+#   define SWIFT_ENUM_NAMED(_type, _name, SWIFT_NAME, _extensibility) SWIFT_ENUM(_type, _name, _extensibility)
+#  endif
 # else
+#  define SWIFT_ENUM(_type, _name, _extensibility) _type _name; enum 
 #  define SWIFT_ENUM_NAMED(_type, _name, SWIFT_NAME, _extensibility) SWIFT_ENUM(_type, _name, _extensibility)
 # endif
 #endif
+#if !defined(SWIFT_ENUM_TAG)
+# if (defined(__cplusplus) && __cplusplus >= 201103L) ||      (defined(__STDC_VERSION__) && __STDC_VERSION__ >= 202311L) ||      __has_feature(objc_fixed_enum)
+#  define SWIFT_ENUM_TAG enum
+# else
+#  define SWIFT_ENUM_TAG
+# endif
+#endif
+#if !defined(SWIFT_ENUM_FWD_DECL)
+# if (defined(__cplusplus) && __cplusplus >= 201103L) ||      (defined(__STDC_VERSION__) && __STDC_VERSION__ >= 202311L) ||      __has_feature(objc_fixed_enum)
+#  define SWIFT_ENUM_FWD_DECL(_type, _name) enum _name : _type;
+# else
+#  define SWIFT_ENUM_FWD_DECL(_type, _name) typedef _type _name;
+# endif
+#endif
 #if !defined(SWIFT_UNAVAILABLE)
 # define SWIFT_UNAVAILABLE __attribute__((unavailable))
 #endif
@@ -210,6 +229,9 @@ typedef unsigned int swift_uint4  __attribute__((__ext_vector_type__(4)));
 #if !defined(SWIFT_AVAILABILITY)
 # define SWIFT_AVAILABILITY(plat, ...) __attribute__((availability(plat, __VA_ARGS__)))
 #endif
+#if !defined(SWIFT_AVAILABILITY_DOMAIN)
+# define SWIFT_AVAILABILITY_DOMAIN(dom, ...) __attribute__((availability(domain: dom, __VA_ARGS__)))
+#endif
 #if !defined(SWIFT_WEAK_IMPORT)
 # define SWIFT_WEAK_IMPORT __attribute__((weak_import))
 #endif
@@ -275,6 +297,49 @@ typedef unsigned int swift_uint4  __attribute__((__ext_vector_type__(4)));
 # define SWIFT_IMPORT_STDLIB_SYMBOL 
 #endif
 #endif
+#if !__has_feature(nullability)
+# define _Nonnull
+# define _Nullable
+# define _Null_unspecified
+#elif !defined(__OBJC__)
+# pragma clang diagnostic ignored "-Wnullability-extension"
+#endif
+#if !__has_feature(nullability_nullable_result)
+# define _Nullable_result _Nullable
+#endif
+#if __has_feature(objc_modules)
+#if __has_warning("-Watimport-in-framework-header")
+#pragma clang diagnostic ignored "-Watimport-in-framework-header"
+#endif
+#endif
+
+#pragma clang diagnostic ignored "-Wproperty-attribute-mismatch"
+#pragma clang diagnostic ignored "-Wduplicate-method-arg"
+#if __has_warning("-Wpragma-clang-attribute")
+# pragma clang diagnostic ignored "-Wpragma-clang-attribute"
+#endif
+#pragma clang diagnostic ignored "-Wunknown-pragmas"
+#pragma clang diagnostic ignored "-Wnullability"
+#pragma clang diagnostic ignored "-Wdollar-in-identifier-extension"
+#pragma clang diagnostic ignored "-Wunsafe-buffer-usage"
+
+#if __has_attribute(external_source_symbol)
+# pragma push_macro("any")
+# undef any
+# pragma clang attribute push(__attribute__((external_source_symbol(language="Swift", defined_in="MarketplaceKitWrapper",generated_declaration))), apply_to=any(function,enum,objc_interface,objc_category,objc_protocol))
+# pragma pop_macro("any")
+#endif
+
+#if defined(__cplusplus)
+extern "C" {
+#endif
+
+#if defined(__cplusplus)
+} // extern "C"
+#endif
+#if __has_attribute(external_source_symbol)
+# pragma clang attribute pop
+#endif
 #if defined(__OBJC__)
 #if __has_feature(objc_modules)
 #if __has_warning("-Watimport-in-framework-header")
@@ -283,7 +348,7 @@ typedef unsigned int swift_uint4  __attribute__((__ext_vector_type__(4)));
 @import ObjectiveC;
 #endif
 
-#endif
+#endif // defined(__OBJC__)
 #pragma clang diagnostic ignored "-Wproperty-attribute-mismatch"
 #pragma clang diagnostic ignored "-Wduplicate-method-arg"
 #if __has_warning("-Wpragma-clang-attribute")
@@ -317,10 +382,11 @@ SWIFT_CLASS("_TtC21MarketplaceKitWrapper21AppDistributorWrapper")
 @interface AppDistributorWrapper : NSObject
 + (void)getCurrentWithCompletionHandler:(void (^ _Nonnull)(enum AppDistributorType, NSString * _Nonnull))CompletionHandler;
 + (void)requestCTToken:(void (^ _Nonnull)(BOOL, NSString * _Nonnull))completionHandler;
++ (void)getEligibilityRegion:(void (^ _Nonnull)(BOOL, NSString * _Nonnull))completionHandler;
 - (nonnull instancetype)init OBJC_DESIGNATED_INITIALIZER;
 @end
 
-#endif
+#endif // defined(__OBJC__)
 #if __has_attribute(external_source_symbol)
 # pragma clang attribute pop
 #endif
diff --git a/Engine/Source/Runtime/IOS/MarketplaceKit/Private/MarketplaceKitWrapperAPI.h b/Engine/Source/Runtime/IOS/MarketplaceKit/Private/MarketplaceKitWrapperAPI.h
new file mode 100644
--- /dev/null
+++ b/Engine/Source/Runtime/IOS/MarketplaceKit/Private/MarketplaceKitWrapperAPI.h
@@ -0,0 +1,26 @@
+// Copyright Epic Games, Inc. All Rights Reserved.
+
+#pragma once
+
+#include <stdbool.h>
+#include <stdint.h>
+
+// C function declarations for the MarketplaceKitWrapper framework.
+// These are weak-imported so the app can run even if the framework is unavailable.
+
+#ifdef __cplusplus
+extern "C"
+{
+#endif
+
+// Swift Int maps to long on arm64 (64-bit), Swift Bool maps to C bool
+typedef void (*MKW_GetCurrentCallback)(long Type, const char* Name);
+typedef void (*MKW_BoolStringCallback)(bool Success, const char* Output);
+
+extern void MKW_GetCurrent(MKW_GetCurrentCallback Callback) __attribute__((weak_import));
+extern void MKW_RequestCTToken(MKW_BoolStringCallback Callback) __attribute__((weak_import));
+extern void MKW_GetEligibilityRegion(MKW_BoolStringCallback Callback) __attribute__((weak_import));
+
+#ifdef __cplusplus
+}
+#endif
diff --git a/Engine/Source/Runtime/IOS/MarketplaceKit/Public/MarketplaceKit.h b/Engine/Source/Runtime/IOS/MarketplaceKit/Public/MarketplaceKit.h
--- a/Engine/Source/Runtime/IOS/MarketplaceKit/Public/MarketplaceKit.h
+++ b/Engine/Source/Runtime/IOS/MarketplaceKit/Public/MarketplaceKit.h
@@ -16,7 +16,7 @@ enum class EMarketplaceType : int32
 	NotAvailable = 5,
 };
 
-const TCHAR* LexToString(EMarketplaceType Value);
+const FString LexToString(EMarketplaceType Value);
 void LexFromString(EMarketplaceType& OutValue, const TCHAR* InValue);
 
 class FMarketplaceKitModule : public IModuleInterface
@@ -28,6 +28,7 @@ public:
 
 	void GetCurrentTypeAsync(TFunction<void(EMarketplaceType Type, const FString& Name)> Callback);
 	static void RequestCTTokenAsync(TDelegate<void(bool Success, FString Output)> Callback);
+	static void GetEligibilityRegionAsync(TDelegate<void(bool Success, FString Output)> Callback);
 
 	/* Returns the MarketplaceType, and BundleId */
 	void GetCurrentType(EMarketplaceType& OutType, FString& OutName);
diff --git a/Engine/Source/Runtime/IOS/MarketplaceKitWrapper/CMakeLists.txt b/Engine/Source/Runtime/IOS/MarketplaceKitWrapper/CMakeLists.txt
--- a/Engine/Source/Runtime/IOS/MarketplaceKitWrapper/CMakeLists.txt
+++ b/Engine/Source/Runtime/IOS/MarketplaceKitWrapper/CMakeLists.txt
@@ -3,9 +3,29 @@
 cmake_minimum_required(VERSION 3.26)
 project(MarketplaceKitWrapper LANGUAGES Swift)
 
-add_library(MarketplaceKitWrapper STATIC MarketplaceKitWrapper.swift)
-set_target_properties(MarketplaceKitWrapper PROPERTIES Swift_MODULE_NAME "MarketplaceKitWrapper")
+set(FRAMEWORK_NAME "MarketplaceKitWrapper")
 
+add_library(${FRAMEWORK_NAME} SHARED MarketplaceKitWrapper.swift)
+set_target_properties(${FRAMEWORK_NAME} PROPERTIES
+	Swift_MODULE_NAME "${FRAMEWORK_NAME}"
+	FRAMEWORK TRUE
+	FRAMEWORK_VERSION A
+	MACOSX_FRAMEWORK_IDENTIFIER "com.epicgames.${FRAMEWORK_NAME}"
+	MACOSX_FRAMEWORK_BUNDLE_VERSION "1.0"
+	MACOSX_FRAMEWORK_SHORT_VERSION_STRING "1.0"
+)
+
+# Set the install name to match the framework bundle structure
+target_link_options(${FRAMEWORK_NAME} PRIVATE
+	"LINKER:-install_name,@rpath/${FRAMEWORK_NAME}.framework/${FRAMEWORK_NAME}"
+)
+
+# Link required system frameworks
+target_link_libraries(${FRAMEWORK_NAME} PRIVATE
+	"-framework Foundation"
+)
+
+# Custom command to generate the ObjC bridging header (still needed for reference)
 add_custom_command(
 	OUTPUT
 		${CMAKE_CURRENT_SOURCE_DIR}/MarketplaceKitWrapper.h
@@ -19,7 +39,7 @@ add_custom_command(
 		-typecheck
 		-target ${CMAKE_OSX_ARCHITECTURES}-apple-ios${CMAKE_OSX_DEPLOYMENT_TARGET}
 		-sdk ${CMAKE_OSX_SYSROOT}
-		-module-name "MarketplaceKitWrapper"
+		-module-name "${FRAMEWORK_NAME}"
 		-emit-objc-header-path ${CMAKE_CURRENT_SOURCE_DIR}/MarketplaceKitWrapper.h
 		MarketplaceKitWrapper.swift
 	COMMAND_EXPAND_LISTS
@@ -30,4 +50,4 @@ add_custom_target(MarketplaceKitWrapperHeader
 		${CMAKE_CURRENT_SOURCE_DIR}/MarketplaceKitWrapper.h
 )
 
-add_dependencies(MarketplaceKitWrapper MarketplaceKitWrapperHeader)
+add_dependencies(${FRAMEWORK_NAME} MarketplaceKitWrapperHeader)
diff --git a/Engine/Source/Runtime/IOS/MarketplaceKitWrapper/MarketplaceKitWrapper.h b/Engine/Source/Runtime/IOS/MarketplaceKitWrapper/MarketplaceKitWrapper.h
new file mode 100644
--- /dev/null
+++ b/Engine/Source/Runtime/IOS/MarketplaceKitWrapper/MarketplaceKitWrapper.h
@@ -0,0 +1,402 @@
+// Copyright Epic Games, Inc. All Rights Reserved.
+// Generated by Apple Swift version 6.3 effective-5.10 (swiftlang-6.3.0.123.5 clang-2100.0.123.102)
+#ifndef MARKETPLACEKITWRAPPER_SWIFT_H
+#define MARKETPLACEKITWRAPPER_SWIFT_H
+#pragma clang diagnostic push
+#pragma clang diagnostic ignored "-Wgcc-compat"
+
+#if !defined(__has_include)
+# define __has_include(x) 0
+#endif
+#if !defined(__has_attribute)
+# define __has_attribute(x) 0
+#endif
+#if !defined(__has_feature)
+# define __has_feature(x) 0
+#endif
+#if !defined(__has_warning)
+# define __has_warning(x) 0
+#endif
+
+#if __has_include(<swift/objc-prologue.h>)
+# include <swift/objc-prologue.h>
+#endif
+
+#pragma clang diagnostic ignored "-Wauto-import"
+#if defined(__OBJC__)
+#include <Foundation/Foundation.h>
+#endif // defined(__OBJC__)
+#if defined(__cplusplus)
+#include <cstdint>
+#include <cstddef>
+#include <cstdbool>
+#include <cstring>
+#include <stdlib.h>
+#include <new>
+#include <type_traits>
+#else
+#include <stdint.h>
+#include <stddef.h>
+#include <stdbool.h>
+#include <string.h>
+#endif
+#if defined(__cplusplus)
+#pragma clang diagnostic push
+#pragma clang diagnostic ignored "-Wnon-modular-include-in-framework-module"
+#if defined(__arm64e__) && __has_include(<ptrauth.h>)
+# include <ptrauth.h>
+#else
+#pragma clang diagnostic push
+#pragma clang diagnostic ignored "-Wreserved-macro-identifier"
+# ifndef __ptrauth_swift_value_witness_function_pointer
+#  define __ptrauth_swift_value_witness_function_pointer(x)
+# endif
+# ifndef __ptrauth_swift_class_method_pointer
+#  define __ptrauth_swift_class_method_pointer(x)
+# endif
+#pragma clang diagnostic pop
+#endif
+#pragma clang diagnostic pop
+#endif
+
+#if !defined(SWIFT_TYPEDEFS)
+# define SWIFT_TYPEDEFS 1
+# if __has_include(<uchar.h>)
+#  include <uchar.h>
+# elif !defined(__cplusplus)
+typedef unsigned char char8_t;
+typedef uint_least16_t char16_t;
+typedef uint_least32_t char32_t;
+# endif
+typedef float swift_float2  __attribute__((__ext_vector_type__(2)));
+typedef float swift_float3  __attribute__((__ext_vector_type__(3)));
+typedef float swift_float4  __attribute__((__ext_vector_type__(4)));
+typedef double swift_double2  __attribute__((__ext_vector_type__(2)));
+typedef double swift_double3  __attribute__((__ext_vector_type__(3)));
+typedef double swift_double4  __attribute__((__ext_vector_type__(4)));
+typedef int swift_int2  __attribute__((__ext_vector_type__(2)));
+typedef int swift_int3  __attribute__((__ext_vector_type__(3)));
+typedef int swift_int4  __attribute__((__ext_vector_type__(4)));
+typedef unsigned int swift_uint2  __attribute__((__ext_vector_type__(2)));
+typedef unsigned int swift_uint3  __attribute__((__ext_vector_type__(3)));
+typedef unsigned int swift_uint4  __attribute__((__ext_vector_type__(4)));
+#endif
+
+#if !defined(SWIFT_PASTE)
+# define SWIFT_PASTE_HELPER(x, y) x##y
+# define SWIFT_PASTE(x, y) SWIFT_PASTE_HELPER(x, y)
+#endif
+#if !defined(SWIFT_METATYPE)
+# define SWIFT_METATYPE(X) Class
+#endif
+#if !defined(SWIFT_CLASS_PROPERTY)
+# if __has_feature(objc_class_property)
+#  define SWIFT_CLASS_PROPERTY(...) __VA_ARGS__
+# else
+#  define SWIFT_CLASS_PROPERTY(...) 
+# endif
+#endif
+#if !defined(SWIFT_RUNTIME_NAME)
+# if __has_attribute(objc_runtime_name)
+#  define SWIFT_RUNTIME_NAME(X) __attribute__((objc_runtime_name(X)))
+# else
+#  define SWIFT_RUNTIME_NAME(X) 
+# endif
+#endif
+#if !defined(SWIFT_COMPILE_NAME)
+# if __has_attribute(swift_name)
+#  define SWIFT_COMPILE_NAME(X) __attribute__((swift_name(X)))
+# else
+#  define SWIFT_COMPILE_NAME(X) 
+# endif
+#endif
+#if !defined(SWIFT_METHOD_FAMILY)
+# if __has_attribute(objc_method_family)
+#  define SWIFT_METHOD_FAMILY(X) __attribute__((objc_method_family(X)))
+# else
+#  define SWIFT_METHOD_FAMILY(X) 
+# endif
+#endif
+#if !defined(SWIFT_NOESCAPE)
+# if __has_attribute(noescape)
+#  define SWIFT_NOESCAPE __attribute__((noescape))
+# else
+#  define SWIFT_NOESCAPE 
+# endif
+#endif
+#if !defined(SWIFT_RELEASES_ARGUMENT)
+# if __has_attribute(ns_consumed)
+#  define SWIFT_RELEASES_ARGUMENT __attribute__((ns_consumed))
+# else
+#  define SWIFT_RELEASES_ARGUMENT 
+# endif
+#endif
+#if !defined(SWIFT_WARN_UNUSED_RESULT)
+# if __has_attribute(warn_unused_result)
+#  define SWIFT_WARN_UNUSED_RESULT __attribute__((warn_unused_result))
+# else
+#  define SWIFT_WARN_UNUSED_RESULT 
+# endif
+#endif
+#if !defined(SWIFT_NORETURN)
+# if __has_attribute(noreturn)
+#  define SWIFT_NORETURN __attribute__((noreturn))
+# else
+#  define SWIFT_NORETURN 
+# endif
+#endif
+#if !defined(SWIFT_CLASS_EXTRA)
+# define SWIFT_CLASS_EXTRA 
+#endif
+#if !defined(SWIFT_PROTOCOL_EXTRA)
+# define SWIFT_PROTOCOL_EXTRA 
+#endif
+#if !defined(SWIFT_ENUM_EXTRA)
+# define SWIFT_ENUM_EXTRA 
+#endif
+#if !defined(SWIFT_CLASS)
+# if __has_attribute(objc_subclassing_restricted)
+#  define SWIFT_CLASS(SWIFT_NAME) SWIFT_RUNTIME_NAME(SWIFT_NAME) __attribute__((objc_subclassing_restricted)) SWIFT_CLASS_EXTRA
+#  define SWIFT_CLASS_NAMED(SWIFT_NAME) __attribute__((objc_subclassing_restricted)) SWIFT_COMPILE_NAME(SWIFT_NAME) SWIFT_CLASS_EXTRA
+# else
+#  define SWIFT_CLASS(SWIFT_NAME) SWIFT_RUNTIME_NAME(SWIFT_NAME) SWIFT_CLASS_EXTRA
+#  define SWIFT_CLASS_NAMED(SWIFT_NAME) SWIFT_COMPILE_NAME(SWIFT_NAME) SWIFT_CLASS_EXTRA
+# endif
+#endif
+#if !defined(SWIFT_RESILIENT_CLASS)
+# if __has_attribute(objc_class_stub)
+#  define SWIFT_RESILIENT_CLASS(SWIFT_NAME) SWIFT_CLASS(SWIFT_NAME) __attribute__((objc_class_stub))
+#  define SWIFT_RESILIENT_CLASS_NAMED(SWIFT_NAME) __attribute__((objc_class_stub)) SWIFT_CLASS_NAMED(SWIFT_NAME)
+# else
+#  define SWIFT_RESILIENT_CLASS(SWIFT_NAME) SWIFT_CLASS(SWIFT_NAME)
+#  define SWIFT_RESILIENT_CLASS_NAMED(SWIFT_NAME) SWIFT_CLASS_NAMED(SWIFT_NAME)
+# endif
+#endif
+#if !defined(SWIFT_PROTOCOL)
+# define SWIFT_PROTOCOL(SWIFT_NAME) SWIFT_RUNTIME_NAME(SWIFT_NAME) SWIFT_PROTOCOL_EXTRA
+# define SWIFT_PROTOCOL_NAMED(SWIFT_NAME) SWIFT_COMPILE_NAME(SWIFT_NAME) SWIFT_PROTOCOL_EXTRA
+#endif
+#if !defined(SWIFT_EXTENSION)
+# define SWIFT_EXTENSION(M) SWIFT_PASTE(M##_Swift_, __LINE__)
+#endif
+#if !defined(OBJC_DESIGNATED_INITIALIZER)
+# if __has_attribute(objc_designated_initializer)
+#  define OBJC_DESIGNATED_INITIALIZER __attribute__((objc_designated_initializer))
+# else
+#  define OBJC_DESIGNATED_INITIALIZER 
+# endif
+#endif
+#if !defined(SWIFT_ENUM_ATTR)
+# if __has_attribute(enum_extensibility)
+#  define SWIFT_ENUM_ATTR(_extensibility) __attribute__((enum_extensibility(_extensibility)))
+# else
+#  define SWIFT_ENUM_ATTR(_extensibility) 
+# endif
+#endif
+#if !defined(SWIFT_ENUM)
+# if (defined(__cplusplus) && __cplusplus >= 201103L) ||      (defined(__STDC_VERSION__) && __STDC_VERSION__ >= 202311L) ||      __has_feature(objc_fixed_enum)
+#  define SWIFT_ENUM(_type, _name, _extensibility) enum _name : _type _name; enum SWIFT_ENUM_ATTR(_extensibility) SWIFT_ENUM_EXTRA _name : _type
+#  if __has_feature(generalized_swift_name)
+#   define SWIFT_ENUM_NAMED(_type, _name, SWIFT_NAME, _extensibility) enum _name : _type _name SWIFT_COMPILE_NAME(SWIFT_NAME); enum SWIFT_COMPILE_NAME(SWIFT_NAME) SWIFT_ENUM_ATTR(_extensibility) SWIFT_ENUM_EXTRA _name : _type
+#  else
+#   define SWIFT_ENUM_NAMED(_type, _name, SWIFT_NAME, _extensibility) SWIFT_ENUM(_type, _name, _extensibility)
+#  endif
+# else
+#  define SWIFT_ENUM(_type, _name, _extensibility) _type _name; enum 
+#  define SWIFT_ENUM_NAMED(_type, _name, SWIFT_NAME, _extensibility) SWIFT_ENUM(_type, _name, _extensibility)
+# endif
+#endif
+#if !defined(SWIFT_ENUM_TAG)
+# if (defined(__cplusplus) && __cplusplus >= 201103L) ||      (defined(__STDC_VERSION__) && __STDC_VERSION__ >= 202311L) ||      __has_feature(objc_fixed_enum)
+#  define SWIFT_ENUM_TAG enum
+# else
+#  define SWIFT_ENUM_TAG
+# endif
+#endif
+#if !defined(SWIFT_ENUM_FWD_DECL)
+# if (defined(__cplusplus) && __cplusplus >= 201103L) ||      (defined(__STDC_VERSION__) && __STDC_VERSION__ >= 202311L) ||      __has_feature(objc_fixed_enum)
+#  define SWIFT_ENUM_FWD_DECL(_type, _name) enum _name : _type;
+# else
+#  define SWIFT_ENUM_FWD_DECL(_type, _name) typedef _type _name;
+# endif
+#endif
+#if !defined(SWIFT_UNAVAILABLE)
+# define SWIFT_UNAVAILABLE __attribute__((unavailable))
+#endif
+#if !defined(SWIFT_UNAVAILABLE_MSG)
+# define SWIFT_UNAVAILABLE_MSG(msg) __attribute__((unavailable(msg)))
+#endif
+#if !defined(SWIFT_AVAILABILITY)
+# define SWIFT_AVAILABILITY(plat, ...) __attribute__((availability(plat, __VA_ARGS__)))
+#endif
+#if !defined(SWIFT_AVAILABILITY_DOMAIN)
+# define SWIFT_AVAILABILITY_DOMAIN(dom, ...) __attribute__((availability(domain: dom, __VA_ARGS__)))
+#endif
+#if !defined(SWIFT_WEAK_IMPORT)
+# define SWIFT_WEAK_IMPORT __attribute__((weak_import))
+#endif
+#if !defined(SWIFT_DEPRECATED)
+# define SWIFT_DEPRECATED __attribute__((deprecated))
+#endif
+#if !defined(SWIFT_DEPRECATED_MSG)
+# define SWIFT_DEPRECATED_MSG(...) __attribute__((deprecated(__VA_ARGS__)))
+#endif
+#if !defined(SWIFT_DEPRECATED_OBJC)
+# if __has_feature(attribute_diagnose_if_objc)
+#  define SWIFT_DEPRECATED_OBJC(Msg) __attribute__((diagnose_if(1, Msg, "warning")))
+# else
+#  define SWIFT_DEPRECATED_OBJC(Msg) SWIFT_DEPRECATED_MSG(Msg)
+# endif
+#endif
+#if defined(__OBJC__)
+#if !defined(IBSegueAction)
+# define IBSegueAction 
+#endif
+#endif
+#if !defined(SWIFT_EXTERN)
+# if defined(__cplusplus)
+#  define SWIFT_EXTERN extern "C"
+# else
+#  define SWIFT_EXTERN extern
+# endif
+#endif
+#if !defined(SWIFT_CALL)
+# define SWIFT_CALL __attribute__((swiftcall))
+#endif
+#if !defined(SWIFT_INDIRECT_RESULT)
+# define SWIFT_INDIRECT_RESULT __attribute__((swift_indirect_result))
+#endif
+#if !defined(SWIFT_CONTEXT)
+# define SWIFT_CONTEXT __attribute__((swift_context))
+#endif
+#if !defined(SWIFT_ERROR_RESULT)
+# define SWIFT_ERROR_RESULT __attribute__((swift_error_result))
+#endif
+#if defined(__cplusplus)
+# define SWIFT_NOEXCEPT noexcept
+#else
+# define SWIFT_NOEXCEPT 
+#endif
+#if !defined(SWIFT_C_INLINE_THUNK)
+# if __has_attribute(always_inline)
+# if __has_attribute(nodebug)
+#  define SWIFT_C_INLINE_THUNK inline __attribute__((always_inline)) __attribute__((nodebug))
+# else
+#  define SWIFT_C_INLINE_THUNK inline __attribute__((always_inline))
+# endif
+# else
+#  define SWIFT_C_INLINE_THUNK inline
+# endif
+#endif
+#if defined(_WIN32)
+#if !defined(SWIFT_IMPORT_STDLIB_SYMBOL)
+# define SWIFT_IMPORT_STDLIB_SYMBOL __declspec(dllimport)
+#endif
+#else
+#if !defined(SWIFT_IMPORT_STDLIB_SYMBOL)
+# define SWIFT_IMPORT_STDLIB_SYMBOL 
+#endif
+#endif
+#if !__has_feature(nullability)
+# define _Nonnull
+# define _Nullable
+# define _Null_unspecified
+#elif !defined(__OBJC__)
+# pragma clang diagnostic ignored "-Wnullability-extension"
+#endif
+#if !__has_feature(nullability_nullable_result)
+# define _Nullable_result _Nullable
+#endif
+#if __has_feature(objc_modules)
+#if __has_warning("-Watimport-in-framework-header")
+#pragma clang diagnostic ignored "-Watimport-in-framework-header"
+#endif
+#endif
+
+#pragma clang diagnostic ignored "-Wproperty-attribute-mismatch"
+#pragma clang diagnostic ignored "-Wduplicate-method-arg"
+#if __has_warning("-Wpragma-clang-attribute")
+# pragma clang diagnostic ignored "-Wpragma-clang-attribute"
+#endif
+#pragma clang diagnostic ignored "-Wunknown-pragmas"
+#pragma clang diagnostic ignored "-Wnullability"
+#pragma clang diagnostic ignored "-Wdollar-in-identifier-extension"
+#pragma clang diagnostic ignored "-Wunsafe-buffer-usage"
+
+#if __has_attribute(external_source_symbol)
+# pragma push_macro("any")
+# undef any
+# pragma clang attribute push(__attribute__((external_source_symbol(language="Swift", defined_in="MarketplaceKitWrapper",generated_declaration))), apply_to=any(function,enum,objc_interface,objc_category,objc_protocol))
+# pragma pop_macro("any")
+#endif
+
+#if defined(__cplusplus)
+extern "C" {
+#endif
+
+#if defined(__cplusplus)
+} // extern "C"
+#endif
+#if __has_attribute(external_source_symbol)
+# pragma clang attribute pop
+#endif
+#if defined(__OBJC__)
+#if __has_feature(objc_modules)
+#if __has_warning("-Watimport-in-framework-header")
+#pragma clang diagnostic ignored "-Watimport-in-framework-header"
+#endif
+@import ObjectiveC;
+#endif
+
+#endif // defined(__OBJC__)
+#pragma clang diagnostic ignored "-Wproperty-attribute-mismatch"
+#pragma clang diagnostic ignored "-Wduplicate-method-arg"
+#if __has_warning("-Wpragma-clang-attribute")
+# pragma clang diagnostic ignored "-Wpragma-clang-attribute"
+#endif
+#pragma clang diagnostic ignored "-Wunknown-pragmas"
+#pragma clang diagnostic ignored "-Wnullability"
+#pragma clang diagnostic ignored "-Wdollar-in-identifier-extension"
+#pragma clang diagnostic ignored "-Wunsafe-buffer-usage"
+
+#if __has_attribute(external_source_symbol)
+# pragma push_macro("any")
+# undef any
+# pragma clang attribute push(__attribute__((external_source_symbol(language="Swift", defined_in="MarketplaceKitWrapper",generated_declaration))), apply_to=any(function,enum,objc_interface,objc_category,objc_protocol))
+# pragma pop_macro("any")
+#endif
+
+#if defined(__OBJC__)
+
+typedef SWIFT_ENUM(NSInteger, AppDistributorType, closed) {
+  AppDistributorTypeAppStore = 0,
+  AppDistributorTypeTestFlight = 1,
+  AppDistributorTypeMarketplace = 2,
+  AppDistributorTypeWeb = 3,
+  AppDistributorTypeOther = 4,
+  AppDistributorTypeNotAvailable = 5,
+};
+
+@class NSString;
+SWIFT_CLASS("_TtC21MarketplaceKitWrapper21AppDistributorWrapper")
+@interface AppDistributorWrapper : NSObject
++ (void)getCurrentWithCompletionHandler:(void (^ _Nonnull)(enum AppDistributorType, NSString * _Nonnull))CompletionHandler;
++ (void)requestCTToken:(void (^ _Nonnull)(BOOL, NSString * _Nonnull))completionHandler;
++ (void)getEligibilityRegion:(void (^ _Nonnull)(BOOL, NSString * _Nonnull))completionHandler;
+- (nonnull instancetype)init OBJC_DESIGNATED_INITIALIZER;
+@end
+
+SWIFT_EXTERN void MKW_GetCurrent(void (* _Nonnull callback)(NSInteger, char const * _Nullable)) SWIFT_NOEXCEPT;
+
+SWIFT_EXTERN void MKW_GetEligibilityRegion(void (* _Nonnull callback)(BOOL, char const * _Nullable)) SWIFT_NOEXCEPT;
+
+SWIFT_EXTERN void MKW_RequestCTToken(void (* _Nonnull callback)(BOOL, char const * _Nullable)) SWIFT_NOEXCEPT;
+
+#endif // defined(__OBJC__)
+#if __has_attribute(external_source_symbol)
+# pragma clang attribute pop
+#endif
+#if defined(__cplusplus)
+#endif
+#pragma clang diagnostic pop
+#endif
diff --git a/Engine/Source/Runtime/IOS/MarketplaceKitWrapper/MarketplaceKitWrapper.swift b/Engine/Source/Runtime/IOS/MarketplaceKitWrapper/MarketplaceKitWrapper.swift
--- a/Engine/Source/Runtime/IOS/MarketplaceKitWrapper/MarketplaceKitWrapper.swift
+++ b/Engine/Source/Runtime/IOS/MarketplaceKitWrapper/MarketplaceKitWrapper.swift
@@ -89,8 +89,91 @@ enum AppDistributorType : Int
 #endif
 		}
 	}
+	
+	public static func getEligibilityRegion(_ completionHandler: @escaping (Bool, String) -> Void)
+	{
+		Task
+		{
+#if canImport(MarketplaceKit)			
+			let result : (Bool, String)
+//#if swift(>=6.3) NOTE: we want this compile to fail if not using Xcode26.4+
+			if #available(iOS 26.4, *)
+			{
+				// Get the eligibility region
+				let optionalRegion = await AppDistributor.eligibilityRegion
+				if let region = optionalRegion
+				{
+					result = (true, region)
+				}
+				else
+				{
+					result = (false, "Unable to obtain region");
+				}
+			}
+			else
+			{
+				result = (false, "IOS 26.4 not available")
+			}
+//#else
+//			result = (false, "Swift 6.3+ not available")
+//#endif  swift(>=6.3)
+			DispatchQueue.main.async 
+			{
+				completionHandler(result.0, result.1)
+			}
+#else // canImport(MarketplaceKit)
+			DispatchQueue.main.async 
+			{
+				completionHandler(false, "MarketplaceKit not imported")
+			}
+#endif // canImport(MarketplaceKit)
+		}
+	}
 }
 
+// MARK: - C-exported entry points for dlopen/dlsym
+
+public typealias MKW_GetCurrentCallback = @convention(c) (Int, UnsafePointer<CChar>?) -> Void
+public typealias MKW_BoolStringCallback = @convention(c) (Bool, UnsafePointer<CChar>?) -> Void
+
+@_cdecl("MKW_GetCurrent")
+public func MKW_GetCurrent(_ callback: @escaping MKW_GetCurrentCallback)
+{
+	AppDistributorWrapper.getCurrent
+	{ (type: AppDistributorType, name: String) in
+		name.withCString
+		{ cStr in
+			callback(type.rawValue, cStr)
+		}
+	}
+}
+
+@_cdecl("MKW_RequestCTToken")
+public func MKW_RequestCTToken(_ callback: @escaping MKW_BoolStringCallback)
+{
+	AppDistributorWrapper.requestCTToken
+	{ (success: Bool, output: String) in
+		output.withCString
+		{ cStr in
+			callback(success, cStr)
+		}
+	}
+}
+
+@_cdecl("MKW_GetEligibilityRegion")
+public func MKW_GetEligibilityRegion(_ callback: @escaping MKW_BoolStringCallback)
+{
+	AppDistributorWrapper.getEligibilityRegion
+	{ (success: Bool, output: String) in
+		output.withCString
+		{ cStr in
+			callback(success, cStr)
+		}
+	}
+}
+
+// MARK: - MarketplaceKit extension
+
 #if canImport(MarketplaceKit)
 @available(iOS 17.4, *)
 private extension AppDistributor
@@ -109,3 +192,4 @@ private extension AppDistributor
 	}
 }
 #endif
+
diff --git a/Engine/Source/Runtime/IOS/MarketplaceKitWrapper/build.sh b/Engine/Source/Runtime/IOS/MarketplaceKitWrapper/build.sh
--- a/Engine/Source/Runtime/IOS/MarketplaceKitWrapper/build.sh
+++ b/Engine/Source/Runtime/IOS/MarketplaceKitWrapper/build.sh
@@ -1,12 +1,92 @@
+#!/bin/bash
 # Copyright Epic Games, Inc. All Rights Reserved.
 
+set -e
+
+FRAMEWORK_NAME="MarketplaceKitWrapper"
+
 rm -rf .build
 mkdir .build
 pushd .build
 # TODO it shouldn't be needed to specify target twice, remove CMAKE_Swift_FLAGS_INIT in future cmake versions
-cmake -G Ninja -DCMAKE_SYSTEM_NAME=iOS -DCMAKE_OSX_DEPLOYMENT_TARGET=15.0 -DCMAKE_OSX_ARCHITECTURES=arm64 -DCMAKE_Swift_FLAGS_INIT="-target arm64-apple-ios15.0" ..
+#XCODE_APP=${XCODE_APP:-/Applications/Xcode-26.4.app}
+#export DEVELOPER_DIR="${XCODE_APP}/Contents/Developer"
+IOS_SYSROOT=$(xcrun --sdk iphoneos --show-sdk-path)
+cmake -G Ninja -DCMAKE_SYSTEM_NAME=iOS -DCMAKE_OSX_SYSROOT="$IOS_SYSROOT" -DCMAKE_OSX_DEPLOYMENT_TARGET=15.0 -DCMAKE_OSX_ARCHITECTURES=arm64 -DCMAKE_Swift_FLAGS_INIT="-target arm64-apple-ios15.0" ..
 ninja -v
 popd
+
+# Generate the ObjC bridging header with Notice.txt prepended
 cat Notice.txt MarketplaceKitWrapper.h > MarketplaceKitWrapperTemp.h
 mv MarketplaceKitWrapperTemp.h MarketplaceKitWrapper.h
-cp .build/libMarketplaceKitWrapper.a .
+
+# Build the .framework bundle structure
+FRAMEWORK_DIR="${FRAMEWORK_NAME}.framework"
+rm -rf "$FRAMEWORK_DIR"
+mkdir -p "$FRAMEWORK_DIR"
+
+# Copy the dynamic library binary (CMake FRAMEWORK target puts it inside .build/<name>.framework/)
+cp ".build/${FRAMEWORK_DIR}/${FRAMEWORK_NAME}" "${FRAMEWORK_DIR}/${FRAMEWORK_NAME}"
+
+# Generate Info.plist
+IOS_SDK_VERSION=$(xcrun --sdk iphoneos --show-sdk-version)
+XCODE_VERSION=$(xcrun xcodebuild -version | head -1 | awk '{print $2}')
+cat > "${FRAMEWORK_DIR}/Info.plist" << PLIST
+<?xml version="1.0" encoding="UTF-8"?>
+<!DOCTYPE plist PUBLIC "-//Apple//DTD PLIST 1.0//EN" "http://www.apple.com/DTDs/PropertyList-1.0.dtd">
+<plist version="1.0">
+<dict>
+	<key>CFBundleDevelopmentRegion</key>
+	<string>English</string>
+	<key>CFBundleExecutable</key>
+	<string>${FRAMEWORK_NAME}</string>
+	<key>CFBundleIdentifier</key>
+	<string>com.epicgames.${FRAMEWORK_NAME}</string>
+	<key>CFBundleInfoDictionaryVersion</key>
+	<string>6.0</string>
+	<key>CFBundlePackageType</key>
+	<string>FMWK</string>
+	<key>CFBundleSignature</key>
+	<string>????</string>
+	<key>CFBundleSupportedPlatforms</key>
+	<array>
+		<string>iPhoneOS</string>
+	</array>
+	<key>CFBundleVersion</key>
+	<string>1.0</string>
+	<key>CFBundleShortVersionString</key>
+	<string>1.0</string>
+	<key>MinimumOSVersion</key>
+	<string>15.0</string>
+	<key>DTPlatformName</key>
+	<string>iphoneos</string>
+	<key>DTPlatformVersion</key>
+	<string>${IOS_SDK_VERSION}</string>
+	<key>DTSDKName</key>
+	<string>iphoneos${IOS_SDK_VERSION}</string>
+	<key>UIRequiredDeviceCapabilities</key>
+	<array>
+		<string>arm64</string>
+	</array>
+</dict>
+</plist>
+PLIST
+
+# Package as .embeddedframework.zip (UBT expects this structure)
+EMBEDDED_DIR="${FRAMEWORK_NAME}.embeddedframework"
+rm -rf "$EMBEDDED_DIR"
+mkdir -p "$EMBEDDED_DIR"
+cp -R "$FRAMEWORK_DIR" "$EMBEDDED_DIR/"
+
+ZIP_NAME="${FRAMEWORK_NAME}.embeddedframework.zip"
+rm -f "$ZIP_NAME"
+zip -r "$ZIP_NAME" "$EMBEDDED_DIR"
+
+# Clean up intermediate directories
+rm -rf "$EMBEDDED_DIR" "$FRAMEWORK_DIR"
+
+# Copy zip to the MarketplaceKit module directory for UE build consumption
+CONSUMER_DIR="../MarketplaceKit"
+cp "$ZIP_NAME" "$CONSUMER_DIR/"
+
+echo "Built ${ZIP_NAME} successfully"
