[iOS CTC Reporting]:
* Engine:
- Move eligiblityRegion to the marketplaceKitModule (like app source)
- Add option to override eligibility region, including forcing it to be empty
- Add additional logging to eligiblity region and store country

diff --git a/Engine/Plugins/Online/IOS/OnlineSubsystemIOS/Source/Private/OnlineStoreIOS.cpp b/Engine/Plugins/Online/IOS/OnlineSubsystemIOS/Source/Private/OnlineStoreIOS.cpp
--- a/Engine/Plugins/Online/IOS/OnlineSubsystemIOS/Source/Private/OnlineStoreIOS.cpp
+++ b/Engine/Plugins/Online/IOS/OnlineSubsystemIOS/Source/Private/OnlineStoreIOS.cpp
@@ -335,38 +335,12 @@ void FOnlineStoreIOS::QueryStoreCountry()
 			const FString CountryString(Country);
 			SharedThis->StoreCountry = CountryString;
 			SharedThis->bIsStoreCountryInitialized.Store(true);
+			
+			UE_LOG_ONLINE_STOREV2(Verbose, TEXT("Obtained store country: %s"), *CountryString);
 		}
 	}];
 }
 
-void FOnlineStoreIOS::QueryEligibilityRegion()
-{
-#if !PLATFORM_TVOS && !PLATFORM_VISIONOS
-	TWeakPtr<FOnlineStoreIOS> WeakThis = AsWeak();
-																													  
-	TDelegate<void(bool, FString)> Callback = TDelegate<void(bool, FString)>::CreateLambda(
-		[WeakThis](bool Success, FString Output)                                                                      
-		{                                                                                                             
-			if (TSharedPtr<FOnlineStoreIOS> SharedThis = WeakThis.Pin())                                          
-			{
-				if (Success)
-				{
-					SharedThis->EligibilityRegion = Output;                                                       
-				}                                                                                                 
-				else
-				{
-					SharedThis->EligibilityRegion = "";                                                       
-					UE_LOG_ONLINE_STOREV2(Warning, TEXT("Unable to obtain eligibility region: %s"), *Output);
-				}
-				SharedThis->bIsEligibilityRegionInitialized.Store(true);
-			}
-		}
-	);
-	
-	FMarketplaceKitModule::GetEligibilityRegionAsync(Callback);
-#endif // !PLATFORM_TVOS && !PLATFORM_VISIONOS
-}
-
 FString FOnlineStoreIOS::GetStoreCountry() const
 {
 	if (bIsStoreCountryInitialized.Load())
@@ -378,17 +352,6 @@ FString FOnlineStoreIOS::GetStoreCountry() const
 	return "";
 }
 
-FString FOnlineStoreIOS::GetEligibilityRegion() const
-{
-	if (bIsEligibilityRegionInitialized.Load())
-	{
-		return EligibilityRegion;
-	}
-	
-	UE_LOG_ONLINE_STOREV2(Warning, TEXT("Attempting to get eligibility region before it is initialized!"));
-	return "";
-}
-
 #if !UE_BUILD_SHIPPING
 class FStoreExec : public FSelfRegisteringExec
 {
diff --git a/Engine/Plugins/Online/IOS/OnlineSubsystemIOS/Source/Private/OnlineSubsystemIOS.cpp b/Engine/Plugins/Online/IOS/OnlineSubsystemIOS/Source/Private/OnlineSubsystemIOS.cpp
--- a/Engine/Plugins/Online/IOS/OnlineSubsystemIOS/Source/Private/OnlineSubsystemIOS.cpp
+++ b/Engine/Plugins/Online/IOS/OnlineSubsystemIOS/Source/Private/OnlineSubsystemIOS.cpp
@@ -208,9 +208,6 @@ bool FOnlineSubsystemIOS::Init()
 		{
 			// Initialize the store country
 			StoreV2Interface->QueryStoreCountry();
-			
-			// Initialize the eligibility region
-			StoreV2Interface->QueryEligibilityRegion();
 		}
 	}
 
diff --git a/Engine/Plugins/Online/IOS/OnlineSubsystemIOS/Source/Public/OnlineStoreIOS.h b/Engine/Plugins/Online/IOS/OnlineSubsystemIOS/Source/Public/OnlineStoreIOS.h
--- a/Engine/Plugins/Online/IOS/OnlineSubsystemIOS/Source/Public/OnlineStoreIOS.h
+++ b/Engine/Plugins/Online/IOS/OnlineSubsystemIOS/Source/Public/OnlineStoreIOS.h
@@ -31,7 +31,6 @@ public:
 	virtual void GetOffers(TArray<FOnlineStoreOfferRef>& OutOffers) const override;
 	virtual TSharedPtr<FOnlineStoreOffer> GetOffer(const FUniqueOfferId& OfferId) const override;
 	virtual FString GetStoreCountry() const override;
-	virtual FString GetEligibilityRegion() const override;
 
 	// FOnlineStoreIOS
 
@@ -65,10 +64,6 @@ public:
 	 */
 	void QueryStoreCountry();
 	
-	/**
-	 * Starts an async request to get the marketplace kit eligibility region
-	 */
-	void QueryEligibilityRegion();
 
 private:
 	
@@ -120,10 +115,7 @@ private:
 	
 	FString StoreCountry;
 	TAtomic<bool> bIsStoreCountryInitialized = false;
-	
-	FString EligibilityRegion;
-	TAtomic<bool> bIsEligibilityRegionInitialized = false;
-	
+
 	/** Reference to the parent subsystem */
 	FOnlineSubsystemIOS* Subsystem;
 };
diff --git a/Engine/Plugins/Online/OnlineSubsystem/Source/Public/Interfaces/OnlineStoreInterfaceV2.h b/Engine/Plugins/Online/OnlineSubsystem/Source/Public/Interfaces/OnlineStoreInterfaceV2.h
--- a/Engine/Plugins/Online/OnlineSubsystem/Source/Public/Interfaces/OnlineStoreInterfaceV2.h
+++ b/Engine/Plugins/Online/OnlineSubsystem/Source/Public/Interfaces/OnlineStoreInterfaceV2.h
@@ -277,9 +277,4 @@ public:
 	* Return store country if it has been fetched.
 	*/
 	virtual FString GetStoreCountry() const { return ""; }
-	
-	/**
-	* Return eligibility region if it has been fetched.
-	*/
-	virtual FString GetEligibilityRegion() const { return ""; }
 };
diff --git a/Engine/Source/Runtime/IOS/MarketplaceKit/Private/MarketplaceKit.cpp b/Engine/Source/Runtime/IOS/MarketplaceKit/Private/MarketplaceKit.cpp
--- a/Engine/Source/Runtime/IOS/MarketplaceKit/Private/MarketplaceKit.cpp
+++ b/Engine/Source/Runtime/IOS/MarketplaceKit/Private/MarketplaceKit.cpp
@@ -13,6 +13,14 @@
 
 DEFINE_LOG_CATEGORY_STATIC(LogMarketplaceKit, Log, All);
 
+static TAutoConsoleVariable<bool> CVarForceEmptyEligibilityRegion(
+	TEXT("MarketplaceKit.ForceEmptyEligibilityRegion"),
+	false,
+	TEXT("Force GetEligibilityRegionStatic to return an empty string, simulating a device with no eligibility region.")
+);
+
+static FMarketplaceKitModule* GPendingEligibilityRegionSelf = nullptr;
+
 static bool IsMarketplaceKitWrapperAvailable()
 {
 	// Weak-linked symbols are nullptr if the framework isn't loaded
@@ -78,17 +86,22 @@ void FMarketplaceKitModule::StartupModule()
 
 	GConfig->GetString(TEXT("MarketplaceKit"), TEXT("TestFlightMarketplaceBundleId"), TestFlightMarketplaceBundleId, GEngineIni);
 	FParse::Value(FCommandLine::Get(), TEXT("TestFlightMarketplaceBundleId="), TestFlightMarketplaceBundleId);
-
+	
+	GConfig->GetString(TEXT("MarketplaceKit"), TEXT("TestOverrideEligibilityRegion"), TestOverrideEligibilityRegion, GEngineIni);
+	FParse::Value(FCommandLine::Get(), TEXT("TestOverrideEligibilityRegion="), TestOverrideEligibilityRegion);
+	
 	if (!IsMarketplaceKitWrapperAvailable())
 	{
 		UE_LOGF(LogMarketplaceKit, Warning, "MarketplaceKitWrapper framework not available, marketplace features will be disabled");
 	}
 
 	CacheValue();
+	GetEligibilityRegionAsync();
 }
 
 void FMarketplaceKitModule::ShutdownModule()
 {
+	GPendingEligibilityRegionSelf = nullptr;
 }
 
 bool FMarketplaceKitModule::SupportsDynamicReloading()
@@ -180,24 +193,69 @@ void FMarketplaceKitModule::RequestCTTokenAsync(TDelegate<void(bool Success, FSt
 	});
 }
 
-void FMarketplaceKitModule::GetEligibilityRegionAsync(TDelegate<void(bool Success, FString Output)> Callback)
+void FMarketplaceKitModule::GetEligibilityRegionAsync()
 {
+	// Check the symbol directly rather than IsMarketplaceKitWrapperAvailable() since that checks
+	// MKW_GetCurrent. If this symbol is unavailable, bCachedEligibilityRegionInitialized will
+	// never be set and GetEligibilityRegionStatic will silently return "".
 	if (!MKW_GetEligibilityRegion)
 	{
-		Callback.ExecuteIfBound(false, FString(TEXT("MarketplaceKitWrapper framework not loaded")));
+		UE_LOGF(LogMarketplaceKit, Warning, "CTC Reporting: Cannot initialize eligibility region.");
 		return;
 	}
+	
+	UE_LOGF(LogMarketplaceKit, Verbose, "CTC Reporting: Querying eligibility region");
 
-	static TDelegate<void(bool, FString)> PendingEligibilityRegionCallback;
-	PendingEligibilityRegionCallback = MoveTemp(Callback);
+	// File-scope static (rather than function-local like other callbacks in this file) so that
+	// ShutdownModule can null it out if the module is unloaded before the callback fires.
+	GPendingEligibilityRegionSelf = this;
 
 	MKW_GetEligibilityRegion([](bool Success, const char* Output)
 	{
-		PendingEligibilityRegionCallback.ExecuteIfBound(Success, FString(UTF8_TO_TCHAR(Output ? Output : "")));
-		PendingEligibilityRegionCallback.Unbind();
+		const FString OutputString(UTF8_TO_TCHAR(Output ? Output : ""));
+
+		if (GPendingEligibilityRegionSelf)
+		{
+			if (Success)
+			{
+				GPendingEligibilityRegionSelf->CachedEligibilityRegion = OutputString;
+			}
+
+			// Store with default seq_cst ordering ensures CachedEligibilityRegion is visible to
+			// any thread that subsequently loads bCachedEligibilityRegionInitialized and sees true.
+			GPendingEligibilityRegionSelf->bCachedEligibilityRegionInitialized.Store(true);
+			GPendingEligibilityRegionSelf = nullptr;
+
+			UE_LOGF(LogMarketplaceKit, Verbose, "CTC Reporting: [%ls] Obtained eligibility region [%ls]",
+					Success ? TEXT("SUCCESS") : TEXT("FAILURE"), *OutputString);
+		}
 	});
 }
 
+FString FMarketplaceKitModule::GetEligibilityRegionStatic()
+{
+	if (CVarForceEmptyEligibilityRegion.GetValueOnAnyThread())
+	{
+		UE_LOGF(LogMarketplaceKit, Verbose, "CTC Reporting: Force empty eligibility region is set, returning empty string.");
+		return FString();
+	}
+
+	if (FMarketplaceKitModule* Module = FModuleManager::LoadModulePtr<FMarketplaceKitModule>(TEXT("MarketplaceKit")))
+	{
+		if (!Module->TestOverrideEligibilityRegion.IsEmpty())
+		{
+			UE_LOGF(LogMarketplaceKit, Verbose, "CTC Reporting: Using overridden eligibility region [%ls]", *Module->TestOverrideEligibilityRegion);
+			return Module->TestOverrideEligibilityRegion;
+		}
+		
+		if (Module->bCachedEligibilityRegionInitialized.Load())
+		{
+			return Module->CachedEligibilityRegion;
+		}
+	}
+	return FString();
+}
+
 void FMarketplaceKitModule::GetCurrentType(EMarketplaceType& OutType, FString& OutName)
 {
 	CacheValue();
diff --git a/Engine/Source/Runtime/IOS/MarketplaceKit/Public/MarketplaceKit.h b/Engine/Source/Runtime/IOS/MarketplaceKit/Public/MarketplaceKit.h
--- a/Engine/Source/Runtime/IOS/MarketplaceKit/Public/MarketplaceKit.h
+++ b/Engine/Source/Runtime/IOS/MarketplaceKit/Public/MarketplaceKit.h
@@ -28,7 +28,7 @@ public:
 
 	void GetCurrentTypeAsync(TFunction<void(EMarketplaceType Type, const FString& Name)> Callback);
 	static void RequestCTTokenAsync(TDelegate<void(bool Success, FString Output)> Callback);
-	static void GetEligibilityRegionAsync(TDelegate<void(bool Success, FString Output)> Callback);
+	static FString GetEligibilityRegionStatic();
 
 	/* Returns the MarketplaceType, and BundleId */
 	void GetCurrentType(EMarketplaceType& OutType, FString& OutName);
@@ -44,6 +44,9 @@ public:
 	bool IsTestFlight();
 	static bool IsTestFlightStatic();
 
+private:
+	void GetEligibilityRegionAsync();
+
 private:
 	bool bCachedTypeValid = false;
 	EMarketplaceType CachedType = EMarketplaceType::NotAvailable;
@@ -52,8 +55,12 @@ private:
 	EMarketplaceType TestFlightMarketplaceType = EMarketplaceType::TestFlight;
 	FString TestFlightMarketplaceBundleId;
 
+	FString TestOverrideEligibilityRegion;
+	FString CachedEligibilityRegion;
+	TAtomic<bool> bCachedEligibilityRegionInitialized = false;
+
 	void CacheValue();
-	
+
 	// Redirects TestFlight to an emulated Type/Name set by CVar for testing.
 	void GetEffectiveType(EMarketplaceType& OutType, FString& OutName);
 };
